WebAPI

diff --git a/Engine/Plugins/Web/WebAPI/Source/WebAPI/Private/Security/WebAPIAuthentication.cpp b/Engine/Plugins/Web/WebAPI/Source/WebAPI/Private/Security/WebAPIAuthentication.cpp
new file mode 100644
--- /dev/null
+++ b/Engine/Plugins/Web/WebAPI/Source/WebAPI/Private/Security/WebAPIAuthentication.cpp
@@ -0,0 +1,162 @@
+﻿// Copyright Epic Games, Inc. All Rights Reserved.
+
+#include "Security/WebAPIAuthentication.h"
+
+#include "WebAPIDeveloperSettings.h"
+#include "WebAPILog.h"
+#include "WebAPISubsystem.h"
+#include "Interfaces/IHttpRequest.h"
+#include "Serialization/JsonSerializer.h"
+
+UWebAPIOAuthSettings::UWebAPIOAuthSettings()
+{
+	SchemeName = TEXT("OAuth");
+}
+
+bool UWebAPIOAuthSettings::IsValid() const
+{
+	return !ClientId.IsEmpty() && !ClientSecret.IsEmpty();
+}
+
+bool FWebAPIOAuthSchemeHandler::HandleHttpRequest(TSharedPtr<IHttpRequest> InRequest, UWebAPIDeveloperSettings* InSettings)
+{
+	check(InSettings);
+
+	// @todo: cache!
+	UWebAPIOAuthSettings* OAuthSettings = nullptr;
+	InSettings->AuthenticationSettings.FindItemByClass<UWebAPIOAuthSettings>(&OAuthSettings);	
+	check(OAuthSettings);
+
+	if(!ensureMsgf(OAuthSettings->IsValid(), TEXT("Authentication settings are missing one or more required properties.")))
+	{
+		return false;
+	}
+
+	if(OAuthSettings->AccessToken.IsEmpty())
+	{
+		return false;
+	}
+
+	InRequest->SetHeader(TEXT("Authorization"), FString::Printf(TEXT("%s %s"), *OAuthSettings->TokenType, *OAuthSettings->AccessToken));
+
+	return true;
+}
+
+bool FWebAPIOAuthSchemeHandler::HandleHttpResponse(EHttpResponseCodes::Type InResponseCode, TSharedPtr<IHttpResponse, ESPMode::ThreadSafe> InResponse, bool bInWasSuccessful, UWebAPIDeveloperSettings* InSettings)
+{
+	check(InSettings);
+
+	// @todo: cache!
+	UWebAPIOAuthSettings* OAuthSettings = nullptr;
+	InSettings->AuthenticationSettings.FindItemByClass<UWebAPIOAuthSettings>(&OAuthSettings);	
+	check(OAuthSettings);
+	
+	if(!ensureMsgf(OAuthSettings->IsValid(), TEXT("Authentication settings are missing one or more required properties.")))
+	{
+		return false;
+	}
+
+	GEngine->GetEngineSubsystem<UWebAPISubsystem>()->MakeHttpRequest(TEXT("POST"), [&, OAuthSettings](const TSharedPtr<IHttpRequest>& InRequest)
+	{
+		//Microsoft login rest URL
+		FStringFormatNamedArguments UrlArgs;
+		UrlArgs.Add(TEXT("TenantId"), OAuthSettings->TenantId);
+		const FString Url = FString::Format(*OAuthSettings->AuthenticationServer, UrlArgs);
+
+		// @todo: move!
+		FString StringPayload = "grant_type=client_credentials";
+		StringPayload.Append("&client_id=" + OAuthSettings->ClientId);
+		StringPayload.Append("&client_secret=" + OAuthSettings->ClientSecret);
+		StringPayload.Append("&resource=https://digitaltwins.azure.net");
+
+		InRequest->SetURL(Url);
+		InRequest->SetContentAsString(StringPayload);
+
+		TMap<FString, FString> Headers;
+		// FInternetAddr
+		FString SchemeName;
+		FParse::SchemeNameFromURI(*Url, SchemeName);
+		FString Host = Url.Replace(*(SchemeName + TEXT("://")), TEXT(""));
+		
+		int32 DelimiterIdx = -1;
+		if(Host.FindChar(TEXT('/'), DelimiterIdx))
+		{
+			Host = Host.Left(DelimiterIdx);
+		}
+
+		Headers.Add("Host", Host);
+		Headers.Add("Content-Type", "application/x-www-form-urlencoded");
+
+		for (const TPair<FString, FString>& Header : Headers)
+		{
+			InRequest->SetHeader(Header.Key, Header.Value);
+		}
+	})
+	.Next([this, OAuthSettings](const TTuple<FHttpResponsePtr, bool>& InResponse)
+	{
+		// Request failed
+		if(!InResponse.Get<bool>())
+		{
+			return false;
+		}
+
+		const FHttpResponsePtr Response = InResponse.Get<FHttpResponsePtr>();
+
+		FString Message;
+		FString AccessTkn;
+		bool bSuccess = false;
+
+		if (Response.IsValid())
+		{
+			const int32 ResponseCode = Response->GetResponseCode();
+			Message = Response->GetContentAsString();
+
+			if (EHttpResponseCodes::IsOk(ResponseCode))
+			{
+				const TSharedRef<TJsonReader<>> JsonReader = TJsonReaderFactory<>::Create(Response->GetContentAsString());
+				TSharedPtr<FJsonObject> JsonObject;
+				if(FJsonSerializer::Deserialize(JsonReader, JsonObject))
+				{
+					const int32 UnixTimeExpire = JsonObject->GetNumberField("expires_on");
+
+					UE_LOG(LogWebAPI, Display, TEXT("Generate token Response success"));
+					OAuthSettings->TokenType = JsonObject->GetStringField("token_type");
+					OAuthSettings->AccessToken = JsonObject->GetStringField("access_token");
+					OAuthSettings->ExpiresOn =  FDateTime::FromUnixTimestamp(UnixTimeExpire);
+					OAuthSettings->SaveConfig();
+					bSuccess = true;
+				}
+				else
+				{
+					UE_LOG(LogWebAPI, Warning, TEXT("Deserialize JSON failed"));
+					UE_LOG(LogWebAPI, Error, TEXT("Authentication failed: Deserialize JSON Response token failed"));
+					Message = "Deserialize JSON Response token failed:" + Message;
+				}
+			}
+			else
+				{
+				UE_LOG(LogWebAPI, Error, TEXT("Authentication failed: Response not valid"));
+				Message = "Response code not valid:" + Message;
+			}
+		}
+		else
+		{
+			UE_LOG(LogWebAPI, Error, TEXT("Authentication failed: Generate token Response not valid"));
+			Message = "Response is null";
+		}
+
+		// @todo:
+		//Redo all the requests that arrived while fetching new token
+		/*
+		if (bUpdatingToken && bSuccess)
+		{
+			UpdatingToken = false;
+			RedoRequests();
+		}
+		*/
+
+		return bSuccess;
+	});
+
+	return true;
+}
diff --git a/Engine/Plugins/Web/WebAPI/Source/WebAPI/Public/WebAPIMessageResponse.h b/Engine/Plugins/Web/WebAPI/Source/WebAPI/Public/WebAPIMessageResponse.h
new file mode 100644
--- /dev/null
+++ b/Engine/Plugins/Web/WebAPI/Source/WebAPI/Public/WebAPIMessageResponse.h
@@ -0,0 +1,19 @@
+﻿// Copyright Epic Games, Inc. All Rights Reserved.
+
+#pragma once
+
+#include "CoreMinimal.h"
+#include "WebAPIMessageResponse.generated.h"
+
+/**
+ * 
+ */
+USTRUCT(BlueprintType)
+struct WEBAPI_API FWebAPIMessageResponse
+{
+	GENERATED_BODY()
+
+	virtual ~FWebAPIMessageResponse() = default;
+
+	virtual const FText& GetMessage() const PURE_VIRTUAL(FWebAPIMessageResponse::GetMessage, return FText::GetEmpty(););
+};
diff --git a/Engine/Plugins/Web/WebAPI/Source/WebAPIBlueprintGraph/Private/K2Node_WebAPIOperation.cpp b/Engine/Plugins/Web/WebAPI/Source/WebAPIBlueprintGraph/Private/K2Node_WebAPIOperation.cpp
new file mode 100644
--- /dev/null
+++ b/Engine/Plugins/Web/WebAPI/Source/WebAPIBlueprintGraph/Private/K2Node_WebAPIOperation.cpp
@@ -0,0 +1,378 @@
+﻿// Copyright Epic Games, Inc. All Rights Reserved.
+
+#include "K2Node_WebAPIOperation.h"
+
+#include "BlueprintActionDatabaseRegistrar.h"
+#include "BlueprintTypePromotion.h"
+#include "EdGraphSchema_K2.h"
+#include "K2Node_StructOperation.h"
+#include "KismetCompiler.h"
+#include "ToolMenu.h"
+#include "WebAPIOperationObject.h"
+#include "Algo/AllOf.h"
+#include "EdGraph/EdGraphNode.h"
+#include "EdGraph/EdGraphPin.h"
+#include "Kismet2/BlueprintEditorUtils.h"
+
+#define LOCTEXT_NAMESPACE "K2Node_WebAPIAsyncOperation"
+
+struct FGetPinName
+{
+	static const FName& GetMessagePinName() { static const FName MessagePinName(TEXT("Message")); return MessagePinName; }
+	static const FName& GetSuccessPinName() { static const FName SuccessPinName(TEXT("OnSuccess")); return SuccessPinName; }
+	static const FName& GetErrorPinName() { static const FName ErrorPinName(TEXT("OnError")); return ErrorPinName; }
+};
+
+UK2Node_WebAPIOperation::UK2Node_WebAPIOperation()
+{
+	ProxyActivateFunctionName = NAME_None;
+
+	// @todo:
+	// 1. on new node spawned, break Request (input) pin by default: FBlueprintEditor::OnSplitStructPin()
+	// 2. option to convert to/from callback styles: UK2Node_DynamicCast::SetPurity(bool bNewPurity)
+	// 3. collapse output responses that inherit MessageResponse to single message output pin
+}
+
+void UK2Node_WebAPIOperation::AllocateDefaultPins()
+{
+	Super::AllocateDefaultPins();
+}
+
+FText UK2Node_WebAPIOperation::GetNodeTitle(ENodeTitleType::Type TitleType) const
+{
+	FText FunctionName = Super::GetNodeTitle(TitleType);
+	FText NamespaceName;
+	FText ServiceName;
+	
+	if (const UFunction* Function = GetFactoryFunction())
+	{
+		NamespaceName = Function->GetOuterUClass()->GetMetaDataText(TEXT("Namespace"));
+		ServiceName = Function->GetOuterUClass()->GetMetaDataText(TEXT("Service"));
+	}
+	
+	if(TitleType == ENodeTitleType::FullTitle)
+	{
+		FFormatNamedArguments Args;
+		Args.Add(TEXT("FunctionName"), FunctionName);
+		Args.Add(TEXT("NamespaceName"), NamespaceName);
+		Args.Add(TEXT("ServiceName"), ServiceName);
+
+		if (NamespaceName.IsEmpty() && ServiceName.IsEmpty())
+		{
+			return FText::Format(LOCTEXT("NodeTitle", "{FunctionName}"), Args);
+		}
+		else if (NamespaceName.IsEmpty())
+		{
+			return FText::Format(LOCTEXT("NodeTitle_WithNamespace", "{FunctionName}\n{ServiceName}"), Args);
+		}
+		else if (ServiceName.IsEmpty())
+		{
+			return FText::Format(LOCTEXT("NodeTitle_WithService", "{FunctionName}\n{NamespaceName}"), Args);
+		}
+		else
+		{
+			return FText::Format(LOCTEXT("NodeTitle_WithNamespaceAndService", "{FunctionName}\n{NamespaceName}: {ServiceName}"), Args);
+		}		
+	}
+	else
+	{
+		return FunctionName;
+	}
+}
+
+void UK2Node_WebAPIOperation::GetNodeContextMenuActions(UToolMenu* Menu, UGraphNodeContextMenuContext* Context) const
+{
+	Super::GetNodeContextMenuActions(Menu, Context);
+
+	if (!Context->bIsDebugging)
+	{
+		{
+			FText MenuEntryTitle = LOCTEXT("MakeCallbackTitle", "Convert to Callback responses");
+			FText MenuEntryTooltip = LOCTEXT("MakeCallbackTooltip", "Removes the execution pins and instead adds callbacks.");
+
+			bool bCanToggleAsyncType = true;
+			auto CanExecuteAsyncTypeToggle = [](const bool bInCanToggleAsyncType)->bool
+			{
+				return bInCanToggleAsyncType;
+			};
+
+			if (OperationAsyncType == EWebAPIOperationAsyncType::Callback)
+			{
+				MenuEntryTitle = LOCTEXT("MakeLatentActionTitle", "Convert to Latent Action");
+				MenuEntryTooltip = LOCTEXT("MakeLatentActionTooltip", "Adds in branching execution pins so that you can separatly handle the responses.");
+
+				const UEdGraphSchema_K2* K2Schema = Cast<UEdGraphSchema_K2>(GetSchema());
+				check(K2Schema != nullptr);
+
+				bCanToggleAsyncType = K2Schema->DoesGraphSupportImpureFunctions(GetGraph()); // @todo 
+				if (!bCanToggleAsyncType)
+				{
+					MenuEntryTooltip = LOCTEXT("CannotMakeLatentActionTooltip", "This graph does not support latent actions.");
+				}
+			}
+
+			FToolMenuSection& Section = Menu->AddSection("K2NodeWebAPIOperation", LOCTEXT("AsyncTypeHeader", "Async Type"));
+			Section.AddMenuEntry(
+				"ToggleAsyncType",
+				MenuEntryTitle,
+				MenuEntryTooltip,
+				FSlateIcon(),
+				FUIAction(
+					FExecuteAction::CreateUObject(const_cast<UK2Node_WebAPIOperation*>(this), &UK2Node_WebAPIOperation::ToggleAsyncType),
+					FCanExecuteAction::CreateStatic(CanExecuteAsyncTypeToggle, bCanToggleAsyncType),
+					FIsActionChecked()
+				)
+			);
+		}
+	}
+}
+
+bool UK2Node_WebAPIOperation::CanPasteHere(const UEdGraph* TargetGraph) const
+{
+	// @todo: unlike default async node, you can paste into a function (not event graph) and this should convert to callback style
+	return Super::CanPasteHere(TargetGraph);
+}
+
+void UK2Node_WebAPIOperation::PostPlacedNewNode()
+{
+	Super::PostPlacedNewNode();
+
+	PostReconstructNode();
+}
+
+void UK2Node_WebAPIOperation::PostPasteNode()
+{
+	Super::PostPasteNode();
+
+	// @todo: handle conversion to/from callback style
+}
+
+void UK2Node_WebAPIOperation::GetMenuActions(FBlueprintActionDatabaseRegistrar& ActionRegistrar) const
+{
+	struct GetMenuActions_Utils
+	{
+		static void SetNodeFunc(UEdGraphNode* NewNode, bool /*bIsTemplateNode*/, TWeakObjectPtr<UFunction> FunctionPtr)
+		{
+			UK2Node_WebAPIOperation* AsyncTaskNode = CastChecked<UK2Node_WebAPIOperation>(NewNode);
+			if (FunctionPtr.IsValid())
+			{
+				const UFunction* Func = FunctionPtr.Get();
+				FObjectProperty* ReturnProp = CastFieldChecked<FObjectProperty>(Func->GetReturnProperty());
+
+				AsyncTaskNode->ProxyFactoryFunctionName = Func->GetFName();
+				AsyncTaskNode->ProxyFactoryClass        = Func->GetOuterUClass();
+				AsyncTaskNode->ProxyClass               = ReturnProp->PropertyClass;				
+			}
+		}
+	};
+
+	UClass* NodeClass = GetClass();
+	ActionRegistrar.RegisterClassFactoryActions<UWebAPIOperationObject>(FBlueprintActionDatabaseRegistrar::FMakeFuncSpawnerDelegate::CreateLambda([NodeClass](const UFunction* FactoryFunc)->UBlueprintNodeSpawner*
+	{
+		UBlueprintNodeSpawner* NodeSpawner = UBlueprintFunctionNodeSpawner::Create(FactoryFunc);
+		check(NodeSpawner != nullptr);
+		
+		NodeSpawner->NodeClass = NodeClass;
+
+		const TWeakObjectPtr<UFunction> FunctionPtr = MakeWeakObjectPtr(const_cast<UFunction*>(FactoryFunc));
+		NodeSpawner->CustomizeNodeDelegate = UBlueprintNodeSpawner::FCustomizeNodeDelegate::CreateStatic(GetMenuActions_Utils::SetNodeFunc, FunctionPtr);
+
+		return NodeSpawner;
+	}));
+}
+
+void UK2Node_WebAPIOperation::PostReconstructNode()
+{
+	Super::PostReconstructNode();
+
+	if(IsValid())
+	{
+		SplitRequestPins(GetRequestPins());
+		SplitResponsePins(GetResponsePins());
+		SplitResponsePins(GetErrorResponsePins());
+	}
+}
+
+void UK2Node_WebAPIOperation::SetAsyncType(EWebAPIOperationAsyncType InAsyncType)
+{
+	if (InAsyncType != OperationAsyncType)
+	{
+		OperationAsyncType = InAsyncType;
+
+		const bool bHasBeenConstructed = (Pins.Num() > 0);
+		if (bHasBeenConstructed)
+		{
+			ReconstructNode();
+		}
+	}
+}
+
+bool UK2Node_WebAPIOperation::IsValid() const
+{
+	return ProxyFactoryFunctionName != NAME_None
+		&& ProxyClass != nullptr
+		&& ((bHasCompilerMessage && ErrorType >= EMessageSeverity::Info) || !bHasCompilerMessage);
+}
+
+UEdGraphPin* UK2Node_WebAPIOperation::FindPinContainingName(
+	const EEdGraphPinDirection& InPinDirection,
+	const FString& InName) const
+{
+	return FindPinByPredicate([&InPinDirection, &InName](const UEdGraphPin* InPin)
+	{
+		return InPin->Direction == InPinDirection
+			&& InPin->PinType.PinCategory != UEdGraphSchema_K2::PC_Exec
+			&& InPin->GetName().Contains(InName);
+	});
+}
+
+TArray<UEdGraphPin*> UK2Node_WebAPIOperation::FindPinsContainingName(
+	const EEdGraphPinDirection& InPinDirection,
+	const FString& InName) const
+{
+	return Pins.FilterByPredicate([&InPinDirection, InName](const UEdGraphPin* InPin)
+	{
+		return InPin->Direction == InPinDirection
+			&& InPin->PinType.PinCategory != UEdGraphSchema_K2::PC_Exec
+			&& InPin->GetName().Contains(InName);
+	});
+}
+
+TArray<UEdGraphPin*> UK2Node_WebAPIOperation::GetRequestPins() const
+{
+	TArray<UEdGraphPin*> RequestPins = FindPinsContainingName(EGPD_Input, TEXT("Request"));
+	ensureMsgf(!RequestPins.IsEmpty(), TEXT("The function must contain a parameter with \"Request\" in the name."));	
+	return RequestPins;
+}
+
+TArray<UEdGraphPin*> UK2Node_WebAPIOperation::GetResponsePins() const
+{
+	auto ResponsePins = FindPinsContainingName(EGPD_Output, TEXT("Response"));
+	ensureMsgf(!ResponsePins.IsEmpty(), TEXT("The Operation must contain a delegate with a parameter containing \"Response\" in the name."));	
+	return ResponsePins;
+}
+
+TArray<UEdGraphPin*> UK2Node_WebAPIOperation::GetErrorResponsePins() const
+{	
+	auto ErrorResponsePins = FindPinsContainingName(EGPD_Output, TEXT("Error"));
+	ensureMsgf(!ErrorResponsePins.IsEmpty(), TEXT("The Operation must contain a delegate with a parameter containing \"Error\" in the name."));	
+	return ErrorResponsePins;
+}
+
+// @todo: recursively split where struct is basically a wrapper for another struct
+void UK2Node_WebAPIOperation::SplitRequestPins(const TArray<UEdGraphPin*>& InPins) const
+{
+	// If UWebAPIOperation base class, the pin won't always be valid, so ignore.
+	if(!ensureAlways(!InPins.IsEmpty() && Algo::AllOf(InPins, [](const UEdGraphPin* InPin) { return InPin; })))
+	{
+		return;
+	}
+
+	for(UEdGraphPin* Pin : InPins)
+	{
+		// Split if not already
+		if (Pin->SubPins.Num() == 0)
+		{
+			if(CanSplitPin(Pin))
+			{
+				const UEdGraphSchema_K2* K2Schema = GetDefault<UEdGraphSchema_K2>();
+				K2Schema->SplitPin(Pin);
+			}
+		}
+
+		// Removes the prefix from the auto-named split pins 
+		if (Pin->SubPins.Num() > 0)
+		{
+			for (UEdGraphPin* SubPin : Pin->SubPins)
+			{
+				FString SubPinName = SubPin->PinName.ToString();
+				CleanupPinNameInline(SubPinName);
+
+				//SubPin->Modify();			
+				SubPin->PinFriendlyName = FText::FromString(SubPinName);
+			}
+		}
+	}
+}
+
+// Splits if the Response contains a single property
+void UK2Node_WebAPIOperation::SplitResponsePins(const TArray<UEdGraphPin*>& InPins) const
+{
+	// Compilation errors or stale nodes should be warnings, not errors
+	if(!ensureAlways(!InPins.IsEmpty() && Algo::AllOf(InPins, [](const UEdGraphPin* InPin) { return InPin; })))
+	{
+		return;	
+	}
+
+	for(UEdGraphPin* Pin : InPins)
+	{
+		if (const UScriptStruct* PinStructType = Cast<UScriptStruct>(Pin->PinType.PinSubCategoryObject.Get()))
+		{
+			uint32 PropertyNum = 0;
+			
+			// Split if not already
+			if (Pin->SubPins.Num() == 0 && CanSplitPin(Pin))
+			{
+				for (TFieldIterator<const FProperty> PropertyIt(PinStructType, EFieldIteratorFlags::ExcludeSuper, EFieldIteratorFlags::ExcludeDeprecated); PropertyIt; ++PropertyIt)
+				{
+					if (PropertyIt->HasAnyPropertyFlags(CPF_Transient))
+					{
+						continue;
+					}
+
+					PropertyNum++;
+				}
+
+				if(PropertyNum == 1)
+				{
+					// Split, result should be single pin (cause it only contained a single property!)
+					const UEdGraphSchema_K2* K2Schema = GetDefault<UEdGraphSchema_K2>();
+					K2Schema->SplitPin(Pin);
+				}
+			}
+
+			// Removes the prefix from the auto-named split pins 
+			if (Pin->SubPins.Num() > 0)
+			{
+				for (UEdGraphPin* SubPin : Pin->SubPins)
+				{
+					FString SubPinName = SubPin->PinName.ToString();
+					CleanupPinNameInline(SubPinName);
+
+					//SubPin->Modify();			
+					SubPin->PinFriendlyName = FText::FromString(SubPinName);
+				}
+			}
+		}
+	}
+}
+
+void UK2Node_WebAPIOperation::CleanupPinNameInline(FString& InPinName)
+{
+	int32 UnderscorePosition = -1;
+	if(InPinName.FindLastChar(TCHAR('_'), UnderscorePosition))
+	{
+		InPinName = InPinName.RightChop(UnderscorePosition + 1);
+	}
+}
+
+void UK2Node_WebAPIOperation::ToggleAsyncType()
+{
+	const FText TransactionTitle =
+		OperationAsyncType == EWebAPIOperationAsyncType::Callback
+		? LOCTEXT("ToggleAsyncTypeToLatentAction", "Convert to Latent Action")
+		: LOCTEXT("ToggleAsyncTypeToCallback", "Convert to Callback responses");
+	const FScopedTransaction Transaction( TransactionTitle );
+	Modify();
+
+	SetAsyncType(StaticCast<EWebAPIOperationAsyncType>(FMath::Abs(StaticCast<int8>(OperationAsyncType) - 1)));
+}
+
+bool UK2Node_WebAPIOperation::ReconnectExecPins(TArray<UEdGraphPin*>& OldPins)
+{
+	// @todo: handle when converting to/from callback style
+	
+	return true;
+}
+
+#undef LOCTEXT_NAMESPACE
diff --git a/Engine/Plugins/Web/WebAPI/Source/WebAPIBlueprintGraph/Private/K2Node_WebAPIOperation.h b/Engine/Plugins/Web/WebAPI/Source/WebAPIBlueprintGraph/Private/K2Node_WebAPIOperation.h
new file mode 100644
--- /dev/null
+++ b/Engine/Plugins/Web/WebAPI/Source/WebAPIBlueprintGraph/Private/K2Node_WebAPIOperation.h
@@ -0,0 +1,71 @@
+﻿// Copyright Epic Games, Inc. All Rights Reserved.
+
+#pragma once
+
+#include "CoreMinimal.h"
+#include "K2Node_BaseAsyncTask.h"
+
+#include "K2Node_WebAPIOperation.generated.h"
+
+class FBlueprintActionDatabaseRegistrar;
+
+UENUM()
+enum class EWebAPIOperationAsyncType : uint8
+{
+	LatentAction = 0,
+	Callback = 1
+};
+
+/**
+ * 
+ */
+UCLASS()
+class WEBAPIBLUEPRINTGRAPH_API UK2Node_WebAPIOperation
+	: public UK2Node_BaseAsyncTask
+{
+	GENERATED_BODY()
+
+public:
+	UK2Node_WebAPIOperation();
+
+	//~ Begin UEdGraphNode Interface.
+	virtual void AllocateDefaultPins() override;
+	virtual FText GetNodeTitle(ENodeTitleType::Type TitleType) const override;
+	virtual void GetNodeContextMenuActions(class UToolMenu* Menu, class UGraphNodeContextMenuContext* Context) const override;
+	virtual void PostPlacedNewNode() override;
+	virtual void PostPasteNode() override;
+	virtual bool CanPasteHere(const UEdGraph* TargetGraph) const override;
+	//~ End UEdGraphNode Interface.
+
+	// UK2Node interface
+	virtual void GetMenuActions(FBlueprintActionDatabaseRegistrar& ActionRegistrar) const override;
+	virtual void PostReconstructNode() override;
+	// End of UK2Node interface
+
+	void SetAsyncType(EWebAPIOperationAsyncType InAsyncType);
+
+	bool IsValid() const;
+
+private:
+	UEdGraphPin* FindPinContainingName(const EEdGraphPinDirection& InPinDirection, const FString& InName) const;
+	TArray<UEdGraphPin*> FindPinsContainingName(const EEdGraphPinDirection& InPinDirection, const FString& InName) const;
+	
+	TArray<UEdGraphPin*> GetRequestPins() const;
+	TArray<UEdGraphPin*> GetResponsePins() const;
+	TArray<UEdGraphPin*> GetErrorResponsePins() const;
+
+	void SplitRequestPins(const TArray<UEdGraphPin*>& InPins) const;	
+	void SplitResponsePins(const TArray<UEdGraphPin*>& InPins) const;
+
+	static void CleanupPinNameInline(FString& InPinName);
+
+	/** Flips the node's Async Type. */
+	void ToggleAsyncType();
+
+	/** Update exec pins when converting between LatentAction and Callback Async Types. */
+	bool ReconnectExecPins(TArray<UEdGraphPin*>& OldPins);
+
+	/** Latent Action is preferred, but not compatible when used in a function, so allow conversion between the two. */
+	UPROPERTY()
+	EWebAPIOperationAsyncType OperationAsyncType = EWebAPIOperationAsyncType::LatentAction;
+};
diff --git a/Engine/Plugins/Web/WebAPI/Source/WebAPIBlueprintGraph/Private/WebAPIBlueprintGraphModule.cpp b/Engine/Plugins/Web/WebAPI/Source/WebAPIBlueprintGraph/Private/WebAPIBlueprintGraphModule.cpp
new file mode 100644
--- /dev/null
+++ b/Engine/Plugins/Web/WebAPI/Source/WebAPIBlueprintGraph/Private/WebAPIBlueprintGraphModule.cpp
@@ -0,0 +1,19 @@
+// Copyright Epic Games, Inc. All Rights Reserved.
+
+#include "WebAPIBlueprintGraphModule.h"
+
+#define LOCTEXT_NAMESPACE "WebAPIBlueprintGraph"
+
+void FWebAPIBlueprintGraphModule::StartupModule()
+{
+
+}
+
+void FWebAPIBlueprintGraphModule::ShutdownModule()
+{
+
+}
+
+#undef LOCTEXT_NAMESPACE
+
+IMPLEMENT_MODULE(FWebAPIBlueprintGraphModule, WebAPIBlueprintGraph)
diff --git a/Engine/Plugins/Web/WebAPI/Source/WebAPIBlueprintGraph/Private/WebAPIBlueprintGraphModule.h b/Engine/Plugins/Web/WebAPI/Source/WebAPIBlueprintGraph/Private/WebAPIBlueprintGraphModule.h
new file mode 100644
--- /dev/null
+++ b/Engine/Plugins/Web/WebAPI/Source/WebAPIBlueprintGraph/Private/WebAPIBlueprintGraphModule.h
@@ -0,0 +1,15 @@
+// Copyright Epic Games, Inc. All Rights Reserved.
+
+#pragma once
+
+#include "IWebAPIBlueprintGraphModule.h"
+
+#include "CoreMinimal.h"
+
+class FWebAPIBlueprintGraphModule final
+    : public IWebAPIBlueprintGraphModuleInterface
+{
+public:
+    void StartupModule() override;
+    void ShutdownModule() override;
+};
diff --git a/Engine/Plugins/Web/WebAPI/Source/WebAPIBlueprintGraph/Public/IWebAPIBlueprintGraphModule.h b/Engine/Plugins/Web/WebAPI/Source/WebAPIBlueprintGraph/Public/IWebAPIBlueprintGraphModule.h
new file mode 100644
--- /dev/null
+++ b/Engine/Plugins/Web/WebAPI/Source/WebAPIBlueprintGraph/Public/IWebAPIBlueprintGraphModule.h
@@ -0,0 +1,11 @@
+// Copyright Epic Games, Inc. All Rights Reserved.
+
+#pragma once
+
+#include "CoreMinimal.h"
+#include "Modules/ModuleInterface.h"
+
+class IWebAPIBlueprintGraphModuleInterface
+    : public IModuleInterface
+{
+};
diff --git a/Engine/Plugins/Web/WebAPI/Source/WebAPIEditor/Private/Details/Widgets/SWebAPISchemaParameterRow.h b/Engine/Plugins/Web/WebAPI/Source/WebAPIEditor/Private/Details/Widgets/SWebAPISchemaParameterRow.h
new file mode 100644
--- /dev/null
+++ b/Engine/Plugins/Web/WebAPI/Source/WebAPIEditor/Private/Details/Widgets/SWebAPISchemaParameterRow.h
@@ -0,0 +1,19 @@
+﻿// Copyright Epic Games, Inc. All Rights Reserved.
+
+#pragma once
+
+#include "CoreMinimal.h"
+#include "SWebAPISchemaTreeTableRow.h"
+#include "Details/ViewModels/WebAPIParameterViewModel.h"
+
+class SWebAPISchemaParameterRow
+	: public SWebAPISchemaTreeTableRow<FWebAPIParameterViewModel>
+{
+public:
+	SLATE_BEGIN_ARGS(SWebAPISchemaParameterRow)
+	{ }
+	SLATE_END_ARGS()
+
+public:
+	void Construct(const FArguments& InArgs, const TSharedRef<FWebAPIParameterViewModel>& InViewModel, const TSharedRef<STableViewBase>& InOwnerTableView);
+};
diff --git a/Engine/Plugins/Web/WebAPI/Source/WebAPIEditor/Private/Dom/WebAPIModelParameter.cpp b/Engine/Plugins/Web/WebAPI/Source/WebAPIEditor/Private/Dom/WebAPIModelParameter.cpp
new file mode 100644
--- /dev/null
+++ b/Engine/Plugins/Web/WebAPI/Source/WebAPIEditor/Private/Dom/WebAPIModelParameter.cpp
@@ -0,0 +1 @@
+﻿// Copyright Epic Games, Inc. All Rights Reserved.
\ No newline at end of file
diff --git a/Engine/Plugins/Web/WebAPI/Source/WebAPIEditor/Private/Dom/WebAPIService.cpp b/Engine/Plugins/Web/WebAPI/Source/WebAPIEditor/Private/Dom/WebAPIService.cpp
new file mode 100644
--- /dev/null
+++ b/Engine/Plugins/Web/WebAPI/Source/WebAPIEditor/Private/Dom/WebAPIService.cpp
@@ -0,0 +1,44 @@
+﻿// Copyright Epic Games, Inc. All Rights Reserved.
+
+#include "Dom/WebAPIService.h"
+
+#include "Dom/WebAPIOperation.h"
+
+#define LOCTEXT_NAMESPACE "WebAPIService"
+
+void UWebAPIService::SetNamespace(const FString& InNamespace)
+{
+	Super::SetNamespace(InNamespace);
+	check(Name.HasTypeInfo());
+
+	if(Name.HasTypeInfo() && !Name.TypeInfo->bIsBuiltinType)
+	{
+		Name.TypeInfo->Namespace = InNamespace;
+	}
+}
+
+void UWebAPIService::Visit(const TFunctionRef<void(IWebAPISchemaObjectInterface*&)> InVisitor)
+{
+	Super::Visit(InVisitor);
+
+	for(const TObjectPtr<UWebAPIOperation>& Operation : Operations)
+	{
+		Operation->Visit(InVisitor);
+	}
+}
+
+#if WITH_EDITOR
+EDataValidationResult UWebAPIService::IsDataValid(TArray<FText>& ValidationErrors)
+{
+	EDataValidationResult ValidationResult = Super::IsDataValid(ValidationErrors);
+
+	for(const TObjectPtr<UWebAPIOperation>& Operation : Operations)
+	{
+		ValidationResult = CombineDataValidationResults(Operation->IsDataValid(ValidationErrors), ValidationResult);
+	}
+
+	return ValidationResult;
+}
+#endif // WITH_EDITOR
+
+#undef LOCTEXT_NAMESPACE
diff --git a/Engine/Plugins/Web/WebAPI/Source/WebAPIEditor/Private/WebAPIEditorStyle.cpp b/Engine/Plugins/Web/WebAPI/Source/WebAPIEditor/Private/WebAPIEditorStyle.cpp
new file mode 100644
--- /dev/null
+++ b/Engine/Plugins/Web/WebAPI/Source/WebAPIEditor/Private/WebAPIEditorStyle.cpp
@@ -0,0 +1,70 @@
+﻿// Copyright Epic Games, Inc. All Rights Reserved.
+
+#include "WebAPIEditorStyle.h"
+
+#include "Interfaces/IPluginManager.h"
+#include "Slate/SlateGameResources.h"
+#include "Styling/SlateStyleMacros.h"
+#include "Styling/SlateStyleRegistry.h"
+#include "Styling/StyleColors.h"
+
+namespace UE
+{
+	namespace WebAPI
+	{
+		namespace Private
+		{
+			const FVector2D Icon16x16(16.0f, 16.0f);
+			const FVector2D Icon20x20(20.0f, 20.0f);
+			const FVector2D Icon40x40(40.0f, 40.0f);
+			const FVector2D Icon64x64(64.0f, 64.0f);
+
+			const FName NAME_StyleName(TEXT("WebAPIEditorStyle"));
+
+			static TUniquePtr<FSlateStyleSet> StyleInstance;	
+		}
+	}
+}
+
+TUniquePtr<FSlateStyleSet> FWebAPIEditorStyle::StyleInstance = nullptr;
+
+using namespace UE::WebAPI;
+
+// This is to fix the issue that SlateStyleMacros like IMAGE_BRUSH look for RootToContentDir but StyleSet->RootToContentDir is how this style is set up
+#define RootToContentDir Private::StyleInstance->RootToContentDir
+#define RootToCoreContentDir Private::StyleInstance->RootToCoreContentDir
+
+void FWebAPIEditorStyle::Register()
+{
+	Private::StyleInstance = MakeUnique<FSlateStyleSet>(Private::NAME_StyleName);
+	Private::StyleInstance->SetCoreContentRoot(FPaths::EngineContentDir() / TEXT("Editor/Slate"));
+	Private::StyleInstance->SetContentRoot(IPluginManager::Get().FindPlugin(TEXT("WebAPI"))->GetContentDir() / TEXT("Editor/Slate"));
+
+	Private::StyleInstance->Set("WebAPI.TreeView.LabelBackground", new CORE_BOX_BRUSH("Common/GroupBorderLight", 6.f/18.f, FStyleColors::Dropdown));
+	
+	FSlateStyleRegistry::RegisterSlateStyle(*Private::StyleInstance.Get());
+}
+
+void FWebAPIEditorStyle::Unregister()
+{
+	FSlateStyleRegistry::UnRegisterSlateStyle(*Private::StyleInstance.Get());
+	Private::StyleInstance.Reset();
+}
+
+FName FWebAPIEditorStyle::GetStyleSetName()
+{
+	return Private::NAME_StyleName;
+}
+
+const ISlateStyle& FWebAPIEditorStyle::Get()
+{
+	check(Private::StyleInstance.IsValid());
+	return *Private::StyleInstance.Get();
+}
+
+#undef IMAGE_BRUSH
+#undef BOX_BRUSH
+#undef BORDER_BRUSH
+#undef DEFAULT_FONT
+#undef CORE_BOX_BRUSH
+#undef CORE_IMAGE_BRUSH
diff --git a/Engine/Plugins/Web/WebAPI/Source/WebAPIEditor/Private/WebAPIEditorWidgetUtilities.cpp b/Engine/Plugins/Web/WebAPI/Source/WebAPIEditor/Private/WebAPIEditorWidgetUtilities.cpp
new file mode 100644
--- /dev/null
+++ b/Engine/Plugins/Web/WebAPI/Source/WebAPIEditor/Private/WebAPIEditorWidgetUtilities.cpp
@@ -0,0 +1,4 @@
+﻿// Copyright Epic Games, Inc. All Rights Reserved.
+
+
+#include "WebAPIEditorWidgetUtilities.h"
diff --git a/Engine/Plugins/Web/WebAPI/Source/WebAPIEditor/Private/WebAPIEditorWidgetUtilities.h b/Engine/Plugins/Web/WebAPI/Source/WebAPIEditor/Private/WebAPIEditorWidgetUtilities.h
new file mode 100644
--- /dev/null
+++ b/Engine/Plugins/Web/WebAPI/Source/WebAPIEditor/Private/WebAPIEditorWidgetUtilities.h
@@ -0,0 +1,26 @@
+﻿// Copyright Epic Games, Inc. All Rights Reserved.
+
+#pragma once
+
+#include "CoreMinimal.h"
+#include "Styling/SlateColor.h"
+#include "UObject/Object.h"
+#include "WebAPIEditorWidgetUtilities.generated.h"
+
+/**
+ * 
+ */
+UCLASS()
+class WEBAPIEDITOR_API UWebAPIEditorWidgetUtilities : public UObject
+{
+	GENERATED_BODY()
+};
+
+
+class FWebAPIEditorWidgetsUtilities
+{
+public:
+	// @todo: verb type should be enum
+	static FSlateColor GetColorForVerb(uint8 InVerb) { return FSlateColor(FColor{255, 128, 64}); }
+	static FName GetIconNameForVerb(uint8 InVerb) { return NAME_None; }
+};
diff --git a/Engine/Plugins/Web/WebAPI/Source/WebAPIEditor/Private/WebApiEditorLog.h b/Engine/Plugins/Web/WebAPI/Source/WebAPIEditor/Private/WebApiEditorLog.h
new file mode 100644
--- /dev/null
+++ b/Engine/Plugins/Web/WebAPI/Source/WebAPIEditor/Private/WebApiEditorLog.h
@@ -0,0 +1,7 @@
+// Copyright Epic Games, Inc. All Rights Reserved.
+
+#pragma once
+
+#include "CoreMinimal.h"
+
+DECLARE_LOG_CATEGORY_EXTERN(LogWebAPIEditor, Log, All);
diff --git a/Engine/Plugins/Web/WebAPI/Source/WebAPIEditor/Public/WebApiJsonUtilities.h b/Engine/Plugins/Web/WebAPI/Source/WebAPIEditor/Public/WebApiJsonUtilities.h
new file mode 100644
--- /dev/null
+++ b/Engine/Plugins/Web/WebAPI/Source/WebAPIEditor/Public/WebApiJsonUtilities.h
@@ -0,0 +1,960 @@
+﻿// Copyright Epic Games, Inc. All Rights Reserved.
+
+#pragma once
+
+#include "CoreMinimal.h"
+#include "Algo/AnyOf.h"
+#include "Dom/JsonObject.h"
+#include "Dom/JsonValue.h"
+
+namespace UE::Json
+{
+	// @note: Concepts are an easier way (than type traits) of checking if a call or specialization is valid
+	namespace Concepts
+	{
+		/** Concept to check if T has NumericLimits, and isn't a bool. */
+		struct CNumerical
+		{
+			template <typename T>
+			auto Requires() -> decltype(
+				TAnd<
+					TIsSame<typename TNumericLimits<T>::NumericType, T>,
+					TNot<TIsSame<T, bool>>>::Value);
+		};
+
+		/** Describes a type that provides a FromJson function. */
+		struct CFromJsonable
+		{
+			template <typename T>
+			auto Requires(T& Val, const TSharedRef<FJsonObject>& Arg) -> decltype(
+				Val.FromJson(Arg)
+			);
+		};
+
+		/** Describes a type the is derived from TMap. */
+		struct CMap 
+		{
+			template <typename T>
+			auto Requires() -> decltype(
+				TOr<
+					TIsTMap<T>,
+					TIsDerivedFrom<T, TMap<typename T::KeyType, typename T::ValueType>>>::Value);
+		};
+	}
+	
+	namespace TypeTraits
+	{
+		/** String-like value types. */
+		template <typename ValueType, typename Enable = void>
+		struct TIsStringLike
+		{
+			enum { Value = false };
+		};
+
+		template <typename ValueType>
+		struct TIsStringLike<
+				ValueType,
+				typename TEnableIf<
+					TOr<
+						TIsSame<ValueType, FString>,
+						TIsSame<ValueType, FName>,
+						TIsSame<ValueType, FText>>::Value>::Type>
+		{
+			enum
+			{
+				Value = true
+			};
+		};
+
+		/** Numeric value types. */
+		template <typename ValueType>
+		using TIsNumeric = TAnd<TModels<Concepts::CNumerical, std::decay_t<ValueType>>>;
+
+		/** ValueType has FromJson. */
+		template <typename ValueType>
+		using THasFromJson = TModels<Concepts::CFromJsonable, std::decay_t<ValueType>>;
+
+		/** ValueType is derived from TMap. */
+		template <typename ValueType>
+		using TIsDerivedFromMap = TModels<Concepts::CMap, std::decay_t<ValueType>>;
+	}
+
+	/** Contains either an object constructed in place, or a reference to an object declared elsewhere. */
+	template<class ObjectType>
+	class TJsonReference
+	{
+	public:
+		using ElementType = ObjectType;
+		static constexpr ESPMode Mode = ESPMode::ThreadSafe;
+			
+		/** Constructs an empty Json Reference. */
+		TJsonReference()
+			: Object(nullptr)
+		{
+		}
+
+		/** Sets the object path and flags it as a pending reference (to be resolved later). */
+		void ResolveDeferred(const FString& InJsonPath)
+		{
+			bHasPendingResolve = true;
+			Path = InJsonPath;
+		}
+
+		/** Attempts to resolve the reference, returns true if successful or already set. */
+		bool TryResolve(const TSharedRef<FJsonObject>& InRootObject, const TFunctionRef<ObjectType(TSharedRef<FJsonObject>&)>& InSetter)
+		{
+			// Already attempted to resolve, return result or nullptr
+			if(!bHasPendingResolve)
+			{
+				return Object.IsValid();
+			}
+
+			TArray<FString> PathSegments;
+			GetPathSegments(PathSegments);
+			PathSegments.RemoveAt(0); // Remove #/
+
+			TSharedPtr<FJsonObject> SubObject = InRootObject;
+			for(const FString& PathSegment : PathSegments)
+			{
+				if(!SubObject->TryGetObjectField(PathSegment, SubObject))
+				{
+					return false;
+				}
+			}
+
+			Object = MoveTemp(InSetter(SubObject.ToSharedRef()));
+			return Object.IsValid();
+		}
+
+		/** Returns the object referenced by this Json Reference, creating it if it doesn't exist. */
+		ObjectType* Get()
+		{
+			if(!Object.IsValid())
+			{
+				Object = MakeShared<ObjectType, Mode>();
+			}
+				
+			return Object.Get();
+		}
+
+		/** Returns the object referenced by this Json Reference, creating it if it doesn't exist. */
+		TSharedPtr<ObjectType> GetShared()
+		{
+			if(!Object.IsValid())
+			{
+				Object = MakeShared<ObjectType, Mode>();
+			}
+				
+			return Object;
+		}
+
+		/** Returns the object referenced by this Json Reference, or nullptr if it doesn't exist. */
+		ObjectType* Get() const
+		{
+			return Object.Get();
+		}
+
+		/** Returns the object referenced by this Json Reference, or nullptr if it doesn't exist. */
+		TSharedPtr<ObjectType> GetShared() const
+		{
+			return Object;
+		}
+
+		/** Set's the object if it's not already set. */
+		bool Set(ObjectType&& InObject)
+		{
+			if(Object.IsValid())
+			{
+				return false;
+			}
+				
+			bHasPendingResolve = false;
+			Object = MoveTemp(InObject);
+			return true;
+		}
+
+		/** Set's the object if it's not already set. */
+		bool Set(TSharedPtr<ObjectType>&& InObject)
+		{
+			if(Object.IsValid())
+			{
+				return false;
+			}
+				
+			bHasPendingResolve = false;
+			Object = MoveTemp(InObject);
+			return true;
+		}
+
+		/** Set's the object if it's not already set. */
+		bool Set(const TSharedPtr<ObjectType>& InObject)
+		{
+			if(Object.IsValid())
+			{
+				return false;
+			}
+				
+			bHasPendingResolve = false;
+			Object = InObject;
+			return true;
+		}
+
+		/** Checks if the underlying object has been set. */
+		bool IsSet() const
+		{
+			return Object.IsValid();
+		}
+
+		/** Checks to see if this is actually pointing to an object. */
+		explicit operator bool() const
+		{
+			return Object != nullptr;
+		}
+
+		/** Checks if the object is valid, or it's pending reference resolution. */
+		bool IsValid() const { return Object != nullptr || bHasPendingResolve; }
+
+		/** Dereferences the object*/
+		ObjectType& operator*() const
+		{
+			check(IsValid());
+			return *Object; 
+		}
+
+		/** Pointer to the underlying object. */
+		ObjectType* operator->() const
+		{
+			check(IsValid());
+			return Get();
+		}
+
+		const FString& GetPath() const
+		{
+			return Path;
+		}
+
+		/** Returns true if there were one or more segments. */
+		bool GetPathSegments(TArray<FString>& OutSegments) const
+		{
+			return Path.ParseIntoArray(OutSegments, TEXT("/")) > 0;
+		}
+
+		FString GetLastPathSegment() const
+		{
+			TArray<FString> PathSegments;
+			if(GetPathSegments(PathSegments))
+			{
+				return PathSegments.Last();
+			}
+			return TEXT("");
+		}
+
+	protected:
+		/** The specified path to the actual object definition. */
+		FString Path;
+
+		/** Flag indicating this is a reference but not yet resolved. */
+		bool bHasPendingResolve = false;
+
+		/** Underling object pointer. */
+		TSharedPtr<ObjectType> Object;
+	};
+		
+	// @note: ideal changes to Json (engine module)
+	// - Support TVariant (same field name, different types)
+	// - Support TOptional for Get* and TryGet*
+
+	// @todo: somehow account for TOptional or not in same call or minimal overload?
+
+	// template <typename ValueType>
+	// void As(const TSharedPtr<FJsonValue>& InJsonValue, ValueType& OutValue)
+	// {
+	// 	unimplemented();
+	// }
+
+	// Numeric
+	template <typename ValueType>
+	constexpr typename TEnableIf<TypeTraits::TIsNumeric<ValueType>::Value, void>::Type
+	As(const TSharedPtr<FJsonValue>& InJsonValue, ValueType& OutValue)
+	{
+		OutValue = InJsonValue->AsNumber();
+	}
+
+	// String
+	template <typename ValueType>
+	constexpr typename TEnableIf<TypeTraits::TIsStringLike<ValueType>::Value, void>::Type
+	As(const TSharedPtr<FJsonValue>& InJsonValue, ValueType& OutValue)
+	{
+		OutValue = InJsonValue->AsString();
+	}
+
+	// Bool
+	template <typename ValueType>
+	constexpr typename TEnableIf<TIsSame<ValueType, bool>::Value, void>::Type
+	As(const TSharedPtr<FJsonValue>& InJsonValue, ValueType& OutValue)
+	{
+		OutValue = InJsonValue->AsBool();
+	}
+
+	// Numeric
+	template <typename ValueType>
+	constexpr typename TEnableIf<TypeTraits::TIsNumeric<ValueType>::Value, bool>::Type
+	TryGet(const TSharedPtr<FJsonValue>& InJsonValue, ValueType& OutValue)
+	{
+		return InJsonValue->TryGetNumber(OutValue);
+	}
+
+	template <typename ValueType>
+	constexpr typename TEnableIf<TypeTraits::TIsNumeric<ValueType>::Value, bool>::Type
+	TryGetField(const TSharedPtr<FJsonObject>& InJsonObject, const FString& InFieldName, ValueType& OutValue)
+	{
+		return InJsonObject->TryGetNumberField(InFieldName, OutValue);
+	}
+
+	// Enum
+	template <typename EnumType>
+	constexpr typename TEnableIf<TIsEnumClass<EnumType>::Value, bool>::Type
+	TryGet(const TSharedPtr<FJsonValue>& InJsonValue, EnumType& OutValue, const TMap<FString, EnumType>& InNameToValue = {})
+	{
+		// Try int first
+		uint8 IntValue;
+		bool bResult = InJsonValue->TryGetNumber(IntValue);
+		if(bResult)
+		{
+			OutValue = StaticCast<EnumType>(IntValue);	
+		}
+		else
+		{
+			if(InNameToValue.IsEmpty())
+			{
+				UE_LOG(LogWebAPI, Error, TEXT("Tried to convert enum from string, but no lookup was provided."));
+				return bResult;  
+			}
+
+			FString StrValue;
+			bResult = InJsonValue->TryGetString(StrValue);
+			if(bResult)
+			{
+				if(const EnumType* NameToValue = InNameToValue.Find(StrValue))
+				{
+					OutValue = *NameToValue;
+				}
+			}
+		}			
+		return bResult;
+	}
+
+	template <typename EnumType>
+	constexpr typename TEnableIf<TIsEnumClass<EnumType>::Value, bool>::Type
+	TryGetField(const TSharedPtr<FJsonObject>& InJsonObject, const FString& InFieldName, EnumType& OutValue, const TMap<FString, EnumType>& InNameToValue = {})
+	{
+		if(const TSharedPtr<FJsonValue> JsonField = InJsonObject->TryGetField(InFieldName))
+		{
+			return TryGet(JsonField, OutValue.Get(), InNameToValue);
+		}
+		return false;
+	}
+
+	// String
+	template <typename ValueType>
+	constexpr typename TEnableIf<TypeTraits::TIsStringLike<ValueType>::Value, bool>::Type
+	TryGet(const TSharedPtr<FJsonValue>& InJsonValue, ValueType& OutValue)
+	{
+		return InJsonValue->TryGetString(OutValue);
+	}
+
+	template <typename ValueType>
+	constexpr typename TEnableIf<TypeTraits::TIsStringLike<ValueType>::Value, bool>::Type
+	TryGetField(const TSharedPtr<FJsonObject>& InJsonObject, const FString& InFieldName, ValueType& OutValue)
+	{
+		return InJsonObject->TryGetStringField(InFieldName, OutValue);
+	}
+
+	// Bool
+	template <typename ValueType>
+	constexpr typename TEnableIf<TIsSame<ValueType, bool>::Value, bool>::Type
+	TryGet(const TSharedPtr<FJsonValue>& InJsonValue, ValueType& OutValue)
+	{
+		return InJsonValue->TryGetBool(OutValue);
+	}
+
+	template <typename ValueType>
+	constexpr typename TEnableIf<TIsSame<ValueType, bool>::Value, bool>::Type
+	TryGetField(const TSharedPtr<FJsonObject>& InJsonObject, const FString& InFieldName, ValueType& OutValue)
+	{
+		return InJsonObject->TryGetBoolField(InFieldName, OutValue);
+	}
+
+	// Array
+	template <typename ContainerType>
+	constexpr typename TEnableIf<TIsTArray<ContainerType>::Value, bool>::Type
+	TryGet(const TSharedPtr<FJsonValue>& InJsonValue, ContainerType& OutValues)
+	{
+		using ValueType = typename ContainerType::ElementType;
+		const TArray<TSharedPtr<FJsonValue>>* Values;
+		if(InJsonValue->TryGetArray(Values))
+		{
+			bool bSuccess = true;
+			if(Values->Num() > 0)
+			{
+				OutValues.Reserve(Values->Num());
+				for(const TSharedPtr<FJsonValue>& JsonValue : *Values)
+				{
+					ValueType& Value = OutValues.Emplace_GetRef();
+					bSuccess &= TryGet(JsonValue, Value);
+				}
+			}
+			return bSuccess;
+		}
+		return false;
+	}
+
+	template <typename ContainerType>
+	constexpr typename TEnableIf<TIsTArray<ContainerType>::Value, bool>::Type
+	TryGetField(const TSharedPtr<FJsonObject>& InJsonObject, const FString& InFieldName, ContainerType& OutValues)
+	{
+		if(const TSharedPtr<FJsonValue> JsonField = InJsonObject->TryGetField(InFieldName))
+		{
+			using ValueType = typename ContainerType::ElementType;
+			return TryGet<ContainerType>(JsonField, OutValues);
+		}
+		return false;
+	}
+
+	// Map
+	template <typename KeyType, typename ValueType>
+	constexpr typename TEnableIf<TypeTraits::TIsStringLike<KeyType>::Value, bool>::Type // Key must be string
+	TryGet(const TSharedPtr<FJsonValue>& InJsonValue, TMap<KeyType, ValueType>& OutValues)
+	{
+		const TSharedPtr<FJsonObject>* JsonObject;
+		if(InJsonValue->TryGetObject(JsonObject))
+		{
+			OutValues.Reserve(JsonObject->Get()->Values.Num());
+			for(TPair<FString, TSharedPtr<FJsonValue, ESPMode::ThreadSafe>>& FieldValuePair : JsonObject->Get()->Values)
+			{
+				ValueType& ItemValue = OutValues.Emplace(MoveTemp(FieldValuePair.Key));
+				TryGet(FieldValuePair.Value.ToSharedRef(), ItemValue);
+			}
+			return true;
+		}
+		return false;
+	}
+
+	template <typename KeyType, typename ValueType>
+	constexpr typename TEnableIf<TypeTraits::TIsStringLike<KeyType>::Value, bool>::Type // Key must be string
+	TryGetField(const TSharedPtr<FJsonObject>& InJsonObject, const FString& InFieldName, TMap<KeyType, ValueType>& OutValues) // Key must be string
+	{
+		if(const TSharedPtr<FJsonValue> JsonField = InJsonObject->TryGetField(InFieldName))
+		{
+			return TryGet(JsonField, OutValues);
+		}
+		return false;
+	}
+	
+	// template <typename ContainerType>
+	// constexpr typename TEnableIf<
+	// 	TAnd<
+	// 		TIsTMap<ContainerType>,
+	// 		TypeTraits::TIsStringLike<typename ContainerType::KeyType>>::Value, bool>::Type // Key must be string
+	// TryGet(const TSharedPtr<FJsonValue>& InJsonValue, ContainerType& OutValues)
+	// {
+	// 	using KeyType = typename ContainerType::KeyType;
+	// 	using ValueType = typename ContainerType::ValueType;
+	//
+	// 	const TSharedPtr<FJsonObject>* Value;
+	// 	if(InJsonValue->TryGetObject(Value))
+	// 	{
+	// 		OutValues.Reserve((*Value)->Values.Num());
+	// 		for(TTuple<FString, TSharedPtr<FJsonValue, ESPMode::ThreadSafe>>& FieldValuePair : (*Value)->Values)
+	// 		{
+	// 			ValueType& ItemValue = OutValues.Add(FieldValuePair.Key, ValueType{});
+	// 			TryGet(FieldValuePair.Value.ToSharedRef(), ItemValue);
+	// 		}
+	// 		return true;
+	// 	}
+	// 	return false;
+	// }
+	//
+	// template <typename ContainerType>
+	// constexpr typename TEnableIf<
+	// 	TAnd<
+	// 		TIsTMap<ContainerType>,
+	// 		TypeTraits::TIsStringLike<typename ContainerType::KeyType>>::Value, bool>::Type
+	// TryGetField(const TSharedPtr<FJsonObject>& InJsonObject, const FString& InFieldName, ContainerType& OutValues) // Key must be string
+	// {
+	// 	if(const TSharedPtr<FJsonValue> JsonField = InJsonObject->TryGetField(InFieldName))
+	// 	{
+	// 		using KeyType = typename ContainerType::KeyType;
+	// 		using ValueType = typename ContainerType::ValueType;
+	// 		return TryGet(JsonField, OutValues);
+	// 	}
+	// 	return false;
+	// }
+
+	// Object (with FromJson)
+	template <typename ValueType>
+	// typename TEnableIf<
+	// 	TAnd<
+	// 		TNot<
+	// 			TAnd<
+	// 				TIsTMap<ValueType>,
+	// 				TIsDerivedFrom<ValueType, TMap<typename ValueType::KeyType, typename ValueType::ValueType>>>>,
+	// 		TypeTraits::THasFromJson<ValueType>>::Value, bool>::Type
+	typename TEnableIf<
+		TAnd<
+			TNot<TypeTraits::TIsDerivedFromMap<ValueType>>,
+			TypeTraits::THasFromJson<ValueType>>::Value, bool>::Type
+	TryGet(const TSharedPtr<FJsonValue>& InJsonValue, ValueType& OutValue);
+
+	template <typename ValueType>
+	typename TEnableIf<
+		TAnd<
+			TNot<TypeTraits::TIsDerivedFromMap<ValueType>>,
+			TypeTraits::THasFromJson<ValueType>>::Value, bool>::Type
+	TryGetField(const TSharedPtr<FJsonObject>& InJsonObject, const FString& InFieldName, ValueType& OutValue);
+
+	// Variant
+	template <typename... ValueTypes>
+	bool TryGet(const TSharedPtr<FJsonValue>& InJsonValue, TVariant<ValueTypes...>& OutValue);
+
+	template <typename... ValueTypes>
+	bool TryGetField(const TSharedPtr<FJsonObject>& InJsonObject, const FString& InFieldName, TVariant<ValueTypes...>& OutValue);
+
+	// UniqueObj
+	template <typename ValueType>
+	bool TryGet(const TSharedPtr<FJsonValue>& InJsonValue, TUniqueObj<ValueType>& OutValue);
+
+	template <typename ValueType>
+	bool TryGetField(const TSharedPtr<FJsonObject>& InJsonObject, const FString& InFieldName, TUniqueObj<ValueType>& OutValue);
+
+	// Object (with FromJson)
+	template <typename ValueType>
+	typename TEnableIf<
+		TAnd<
+			TNot<TypeTraits::TIsDerivedFromMap<ValueType>>,
+			TypeTraits::THasFromJson<ValueType>>::Value, bool>::Type
+	TryGet(const TSharedPtr<FJsonValue>& InJsonValue, ValueType& OutValue)
+	{
+		const TSharedPtr<FJsonObject>* Value;
+		if(InJsonValue->TryGetObject(Value))
+		{
+			return OutValue.FromJson(Value->ToSharedRef());
+		}
+		return false;
+	}
+
+	template <typename ValueType>
+	typename TEnableIf<
+		TAnd<
+			TNot<TypeTraits::TIsDerivedFromMap<ValueType>>,
+			TypeTraits::THasFromJson<ValueType>>::Value, bool>::Type
+	TryGetField(const TSharedPtr<FJsonObject>& InJsonObject, const FString& InFieldName, ValueType& OutValue)
+	{
+		if(const TSharedPtr<FJsonValue> JsonField = InJsonObject->TryGetField(InFieldName))
+		{
+			return TryGet(JsonField, OutValue);
+		}
+		return false;
+	}
+
+	// Object (without FromJson)
+	template <typename ValueType>
+	typename TEnableIf<
+		TAnd<
+			TNot<TIsTArray<ValueType>>,
+			TNot<TypeTraits::TIsStringLike<ValueType>>,
+			TNot<TIsPODType<ValueType>>,
+			TNot<TypeTraits::THasFromJson<ValueType>>>::Value, bool>::Type
+	TryGet(const TSharedPtr<FJsonValue>& InJsonValue, ValueType& OutValue)
+	{
+		static_assert(TypeTraits::THasFromJson<ValueType>::Value, "ValueType must implement FromJson");
+		return false;
+	}
+
+	// TIsTMap<ContainerType>,
+	template <typename ValueType>
+	typename TEnableIf<
+		TAnd<
+			TNot<TIsTArray<ValueType>>,
+			TNot<TypeTraits::TIsStringLike<ValueType>>,
+			TNot<TIsPODType<ValueType>>,
+			TNot<TypeTraits::THasFromJson<ValueType>>>::Value, bool>::Type
+	TryGetField(const TSharedPtr<FJsonObject>& InJsonObject, const FString& InFieldName, ValueType& OutValue)
+	{
+		static_assert(TypeTraits::THasFromJson<ValueType>::Value, "ValueType must implement FromJson");
+		return false;
+	}
+
+	// UniqueObj
+	template <typename ValueType>
+	bool TryGet(const TSharedPtr<FJsonValue>& InJsonValue, TUniqueObj<ValueType>& OutValue)
+	{
+		return Json::TryGet(InJsonValue, OutValue.Get());
+	}
+	
+	template <typename ValueType>
+	bool TryGetField(const TSharedPtr<FJsonObject>& InJsonObject, const FString& InFieldName, TUniqueObj<ValueType>& OutValue)
+	{
+		if(const TSharedPtr<FJsonValue> JsonField = InJsonObject->TryGetField(InFieldName))
+		{
+			return TryGet(JsonField, OutValue.Get());
+		}
+		return false;
+	}
+
+	// TJsonReference
+	template <typename ValueType>
+	bool TryGet(const TSharedPtr<FJsonValue>& InJsonValue, TJsonReference<ValueType>& OutValue)
+	{
+		const TSharedPtr<FJsonObject>* JsonObject;
+		if(InJsonValue->TryGetObject(JsonObject))
+		{
+			FString ReferencePath;
+			if(JsonObject->Get()->TryGetStringField(TEXT("$ref"), ReferencePath))
+			{
+				OutValue.ResolveDeferred(ReferencePath);
+				return true;
+			}
+		}
+		return Json::TryGet(InJsonValue, *OutValue.Get());
+	}
+	
+	template <typename ValueType>
+	bool TryGetField(const TSharedPtr<FJsonObject>& InJsonObject, const FString& InFieldName, TJsonReference<ValueType>& OutValue)
+	{
+		if(const TSharedPtr<FJsonValue> JsonField = InJsonObject->TryGetField(InFieldName))
+		{
+			return TryGet(JsonField, OutValue);
+		}
+		return false;
+	}
+
+	// SharedPtr
+	template <typename ValueType>
+	bool TryGet(const TSharedPtr<FJsonValue>& InJsonValue, TSharedPtr<ValueType>& OutValue)
+	{
+		// Initialize if invalid
+		if(!OutValue.IsValid())
+		{
+			OutValue = MakeShared<ValueType>();
+		}
+			
+		return Json::TryGet(InJsonValue, *OutValue.Get());
+	}
+	
+	template <typename ValueType>
+	bool TryGetField(const TSharedPtr<FJsonObject>& InJsonObject, const FString& InFieldName, TSharedPtr<ValueType>& OutValue)
+	{
+		if(const TSharedPtr<FJsonValue> JsonField = InJsonObject->TryGetField(InFieldName))
+		{
+			return TryGet<ValueType>(JsonField, OutValue);
+		}
+		return false;
+	}
+
+	// SharedRef
+	template <typename ValueType>
+	bool TryGet(const TSharedPtr<FJsonValue>& InJsonValue, TSharedRef<ValueType>& OutValue)
+	{
+		return Json::TryGet(InJsonValue, OutValue.Get());
+	}
+	
+	template <typename ValueType>
+	bool TryGetField(const TSharedPtr<FJsonObject>& InJsonObject, const FString& InFieldName, TSharedRef<ValueType>& OutValue)
+	{
+		if(const TSharedPtr<FJsonValue> JsonField = InJsonObject->TryGetField(InFieldName))
+		{
+			return TryGet(JsonField, OutValue);
+		}
+		return false;
+	}
+
+	// Optional
+	template <typename ValueType>
+	constexpr typename TEnableIf<!TIsEnumClass<ValueType>::Value, bool>::Type
+	TryGet(const TSharedPtr<FJsonValue>& InJsonValue, TOptional<ValueType>& OutValue)
+	{
+		ValueType Value;
+		if(TryGet<TDecay<ValueType>>(InJsonValue, Value))
+		{
+			OutValue = Value;
+			return true;
+		}
+		return false;
+	}
+	
+	template <typename ValueType>
+	constexpr typename TEnableIf<!TIsEnumClass<ValueType>::Value, bool>::Type
+	TryGetField(const TSharedPtr<FJsonObject>& InJsonObject, const FString& InFieldName, TOptional<ValueType>& OutValue)
+	{
+		if(const TSharedPtr<FJsonValue> JsonField = InJsonObject->TryGetField(InFieldName))
+		{
+			if(!OutValue.IsSet())
+			{
+				OutValue = ValueType{};
+			}
+				
+			return TryGet(JsonField, OutValue.GetValue());
+		}
+		return false;
+	}
+
+	// Optional Enum
+	template <typename EnumType>
+	constexpr typename TEnableIf<TIsEnumClass<EnumType>::Value, bool>::Type
+	TryGet(const TSharedPtr<FJsonValue>& InJsonValue, TOptional<EnumType>& OutValue, const TMap<FString, EnumType>& InNameToValue = {})
+	{
+		EnumType Value;
+		if(TryGet<TDecay<EnumType>>(InJsonValue, Value, InNameToValue))
+		{
+			OutValue = Value;
+			return true;
+		}
+		return false;
+	}
+	
+	template <typename EnumType>
+	constexpr typename TEnableIf<TIsEnumClass<EnumType>::Value, bool>::Type
+	TryGetField(const TSharedPtr<FJsonObject>& InJsonObject, const FString& InFieldName, TOptional<EnumType>& OutValue, const TMap<FString, EnumType>& InNameToValue = {})
+	{
+		if(const TSharedPtr<FJsonValue> JsonField = InJsonObject->TryGetField(InFieldName))
+		{
+			if(!OutValue.IsSet())
+			{
+				OutValue = EnumType{};
+			}
+				
+			return TryGet(JsonField, OutValue.GetValue(), InNameToValue);
+		}
+		return false;
+	}
+
+	// Optional UniqueObj
+	template <typename ValueType>
+	bool TryGet(const TSharedPtr<FJsonValue>& InJsonValue, TOptional<TUniqueObj<ValueType>>& OutValue)
+	{
+		if(OutValue.IsSet())
+		{
+			TUniqueObj<ValueType>& Value = OutValue.GetValue();
+			if(TryGet(InJsonValue, OutValue.GetValue()))
+			{
+				OutValue = Value;
+				return true;
+			}
+		}
+		else
+		{
+			TUniqueObj<ValueType> Value;
+			if(TryGet(InJsonValue, Value))
+			{
+				OutValue = Value;
+				return true;
+			}
+		}	
+
+		return false;
+	}
+	
+	template <typename ValueType>
+	bool TryGetField(const TSharedPtr<FJsonObject>& InJsonObject, const FString& InFieldName, TOptional<TUniqueObj<ValueType>>& OutValue)
+	{
+		if(const TSharedPtr<FJsonValue> JsonField = InJsonObject->TryGetField(InFieldName))
+		{
+			return TryGet(JsonField, OutValue);
+		}
+		return false;
+	}
+
+	// Optional SharedRef
+	template <typename ValueType>
+	bool TryGet(const TSharedPtr<FJsonValue>& InJsonValue, TOptional<TSharedRef<ValueType>>& OutValue)
+	{
+		TSharedRef<ValueType> Value = MakeShared<ValueType>();
+		if(TryGet(InJsonValue, Value))
+		{
+			OutValue = Value;
+			return true;
+		}
+		return false;
+	}
+	
+	template <typename ValueType>
+	bool TryGetField(const TSharedPtr<FJsonObject>& InJsonObject, const FString& InFieldName, TOptional<TSharedRef<ValueType>>& OutValue)
+	{
+		if(const TSharedPtr<FJsonValue> JsonField = InJsonObject->TryGetField(InFieldName))
+		{
+			return TryGet(JsonField, OutValue);
+		}
+		return false;
+	}
+
+	// JsonObject
+	template <typename ValueType, typename = typename TEnableIf<TIsSame<ValueType, TSharedPtr<FJsonObject>>::Value>::Type>
+	bool TryGet(const TSharedPtr<FJsonValue>& InJsonValue, ValueType& OutValue)
+	{
+		const TSharedPtr<FJsonObject>* Value;
+		if(InJsonValue->TryGetObject(Value))
+		{
+			OutValue = *Value;
+			return true;
+		}
+		return false;
+	}
+
+	template <typename ValueType, typename = typename TEnableIf<TIsSame<ValueType, TSharedPtr<FJsonObject>>::Value>::Type>
+	bool TryGetField(const TSharedPtr<FJsonObject>& InJsonObject, const FString& InFieldName, TSharedPtr<FJsonObject>& OutValue)
+	{
+		if(const TSharedPtr<FJsonValue> JsonField = InJsonObject->TryGetField(InFieldName))
+		{
+			return TryGet<>(JsonField, OutValue);
+		}
+		return false;
+	}
+
+	// Variant
+	struct FVariantValueVisitor
+	{
+		explicit FVariantValueVisitor(const TSharedPtr<FJsonValue>& InJsonValue)
+			: JsonValue(InJsonValue)
+		{
+		}
+
+		template <typename ValueType>
+		bool operator()(ValueType& OutValue)
+		{
+			RunCount++;
+			if(bFoundMatch)
+			{
+				return bFoundMatch;
+			}
+
+			ValueType TempValue;
+			if(UE::Json::TryGet(JsonValue, TempValue))
+			{
+				OutValue = MoveTemp(TempValue);
+				bFoundMatch = true;
+				return true;
+			}
+
+			return false;
+		}
+
+		// @todo: for debug, remove
+		uint32 RunCount = 0;
+		bool bFoundMatch = false;
+		const TSharedPtr<FJsonValue> JsonValue;
+	};
+
+	struct FVariantObjectVisitor
+	{
+		explicit FVariantObjectVisitor(const TSharedPtr<FJsonObject>& InJsonObject, const FString& InFieldName)
+			: JsonObject(InJsonObject)
+			, FieldName(InFieldName)
+		{
+		}
+
+		template <typename ValueType>
+		bool operator()(ValueType& OutValue)
+		{
+			if(bFoundMatch)
+			{
+				return bFoundMatch;
+			}
+				
+			if(UE::Json::TryGetField(JsonObject, FieldName, OutValue))
+			{
+				bFoundMatch = true;
+				return true;
+			}
+
+			return false;
+		}
+
+		bool bFoundMatch = false;
+		const TSharedPtr<FJsonObject> JsonObject;
+		const FString FieldName;
+	};
+
+	template <typename... ValueTypes>
+	bool TryGet(const TSharedPtr<FJsonValue>& InJsonValue, TVariant<ValueTypes...>& OutValue)
+	{
+		FVariantValueVisitor Visitor(InJsonValue);
+
+		// @todo: asked Stefan Boberg about this, only ever tests a single type
+		return Visit(Visitor, OutValue);
+	}
+	
+	template <typename... ValueTypes>
+	bool TryGetField(const TSharedPtr<FJsonObject>& InJsonObject, const FString& InFieldName, TVariant<ValueTypes...>& OutValue)
+	{
+		if(const TSharedPtr<FJsonValue> JsonField = InJsonObject->TryGetField(InFieldName))
+		{
+			return TryGet(JsonField, OutValue);
+		}
+		return false;
+	}
+
+	/*
+		// Optional Object Ptr
+		template <typename ValueType>
+		constexpr
+		typename TEnableIf<!TIsTrivial<ValueType>::Value, bool>::Type
+		TryGetField(const TSharedPtr<FJsonObject>& InJsonObject, const FString& InFieldName, TOptional<TUniqueObj<ValueType>>& OutValue)
+		{
+		const TSharedPtr<FJsonObject>* Value;
+		if(InJsonObject->TryGetObjectField(InFieldName, Value))
+		{
+		OutValue.Emplace()->FromJson(Value->ToSharedRef());
+		return true;
+		}
+		return false;
+		}
+
+		// Optional Array of Strings
+		template <typename ValueType, typename JsonObjectType>
+		constexpr
+		typename TEnableIf<TIsSame<ValueType, FString>::Value, bool>::Type
+		TryGetField(const TSharedPtr<JsonObjectType>& InJsonObject, const FString& InFieldName, TOptional<TArray<ValueType>>& OutValue)
+		{
+		TArray<ValueType> Value;
+		if(InJsonObject->TryGetStringArrayField(InFieldName, Value))
+		{
+		OutValue.Emplace(MoveTemp(Value));
+		return true; 
+		}
+
+		// Optional Array of Variant Ptr
+		template <typename... ValueTypes>
+		constexpr bool
+		TryGetField(const TSharedPtr<FJsonObject>& InJsonObject, const FString& InFieldName, TOptional<TArray<TUniqueObj<TVariant<ValueTypes...>>>>& OutValue)
+		{
+		const TArray<TSharedPtr<FJsonValue>>* ValueArray;
+		if(InJsonObject->TryGetArrayField(InFieldName, ValueArray))
+		{
+		if(ValueArray->Num() > 0)
+		{
+		OutValue.Emplace({}); // initialize array
+		static constexpr std::initializer_list<ValueTypes...> VariantTypes;
+		for(const TSharedPtr<FJsonValue>& JsonValue : *ValueArray)
+		{
+		for(auto& VariantType : VariantTypes)
+		{
+		using ValueType = decltype(VariantType);
+		ValueType Value;
+		if(TryGet(JsonValue, Value))
+		{
+		// It was this type, so exit loop!
+		OutValue->Emplace(MoveTemp(Value));
+		break;
+		}
+		}
+		}
+		}
+		return true;
+		}
+		return false;
+		}
+		*/
+}
+
+PRAGMA_ENABLE_OPTIMIZATION
