WebAPI cisfix

diff --git a/Engine/Plugins/Web/WebAPI/Source/WebAPI/Public/WebAPIHttpRequest.h b/Engine/Plugins/Web/WebAPI/Source/WebAPI/Public/WebAPIHttpRequest.h
--- a/Engine/Plugins/Web/WebAPI/Source/WebAPI/Public/WebAPIHttpRequest.h
+++ b/Engine/Plugins/Web/WebAPI/Source/WebAPI/Public/WebAPIHttpRequest.h
@@ -1,4 +1,6 @@
-﻿#pragma once
+﻿// Copyright Epic Games, Inc. All Rights Reserved.
+
+#pragma once
 
 #include "CoreMinimal.h"
 #include "HttpModule.h"
diff --git a/Engine/Plugins/Web/WebAPI/Source/WebAPIEditor/Public/WebApiJsonUtilities.h b/Engine/Plugins/Web/WebAPI/Source/WebAPIEditor/Public/WebApiJsonUtilities.h
--- a/Engine/Plugins/Web/WebAPI/Source/WebAPIEditor/Public/WebApiJsonUtilities.h
+++ b/Engine/Plugins/Web/WebAPI/Source/WebAPIEditor/Public/WebApiJsonUtilities.h
@@ -956,5 +956,3 @@ namespace UE::Json
 		}
 		*/
 }
-
-PRAGMA_ENABLE_OPTIMIZATION
