WebAPI missing modules

diff --git a/Engine/Plugins/Web/WebAPI/Source/WebAPI/Private/WebAPIUtilities.cpp b/Engine/Plugins/Web/WebAPI/Source/WebAPI/Private/WebAPIUtilities.cpp
new file mode 100644
--- /dev/null
+++ b/Engine/Plugins/Web/WebAPI/Source/WebAPI/Private/WebAPIUtilities.cpp
@@ -0,0 +1,8 @@
+﻿// Copyright Epic Games, Inc. All Rights Reserved.
+
+#include "WebAPIUtilities.h"
+
+const FText& UWebAPIUtilities::GetResponseMessage(const FWebAPIMessageResponse& MessageResponse)
+{
+	return MessageResponse.GetMessage();
+}
diff --git a/Engine/Plugins/Web/WebAPI/Source/WebAPIEditor/Private/Details/ViewModels/WebAPIEnumViewModel.cpp b/Engine/Plugins/Web/WebAPI/Source/WebAPIEditor/Private/Details/ViewModels/WebAPIEnumViewModel.cpp
new file mode 100644
--- /dev/null
+++ b/Engine/Plugins/Web/WebAPI/Source/WebAPIEditor/Private/Details/ViewModels/WebAPIEnumViewModel.cpp
@@ -0,0 +1,126 @@
+﻿// Copyright Epic Games, Inc. All Rights Reserved.
+
+#include "WebAPIEnumViewModel.h"
+
+#include "Dom/WebAPIEnum.h"
+#include "Dom/WebAPIType.h"
+
+#define LOCTEXT_NAMESPACE "WebAPIEnumViewModel"
+
+TSharedRef<FWebAPIEnumValueViewModel> FWebAPIEnumValueViewModel::Create(const TSharedRef<IWebAPIViewModel>& InParentViewModel, UWebAPIEnumValue* InValue, int32 InIndex)
+{
+	TSharedRef<FWebAPIEnumValueViewModel> ViewModel = MakeShared<FWebAPIEnumValueViewModel>(FPrivateToken{}, InParentViewModel, InValue, InIndex);
+	ViewModel->Initialize();
+
+	return ViewModel;
+}
+
+FWebAPIEnumValueViewModel::FWebAPIEnumValueViewModel(FPrivateToken, const TSharedRef<IWebAPIViewModel>& InParentViewModel, UWebAPIEnumValue* InValue, int32 InIndex)
+	: Enum(InParentViewModel)
+	, Value(InValue)
+	, Index(InIndex)
+{
+	check(Value.IsValid());
+}
+
+void FWebAPIEnumValueViewModel::Initialize()
+{
+	CachedLabel = FText::FromString(Value->Name.GetDisplayName());
+
+	const FString Prefix = TEXT("EV");
+	
+	TArray<FText> TooltipArgs;
+	TooltipArgs.Emplace(LOCTEXT("EnumValueTooltip_Header", "Enum Value").ToUpper());
+	TooltipArgs.Emplace(FText::Format(LOCTEXT("EnumValueTooltip_MemberName", "MemberName: {0}"), FText::FromString(Value->Name.NameInfo.ToMemberName(Prefix))));
+	
+	if(!Value->Description.IsEmpty())
+	{
+		TooltipArgs.Emplace(FText::Format(LOCTEXT("EnumValueTooltip_Description", "Description: {0}"), FText::FromString(Value->Description)));
+	}
+
+	TooltipArgs.Emplace(FText::Format(LOCTEXT("EnumValueTooltip_JsonName", "JsonName: {0}"), FText::FromString(Value->Name.GetJsonName())));
+
+	CachedTooltip = FText::Join(FText::FromString(TEXT("\n")), TooltipArgs);
+}
+
+TSharedPtr<IWebAPIViewModel> FWebAPIEnumValueViewModel::GetParent()
+{
+	return Enum;
+}
+
+TSharedRef<FWebAPIEnumViewModel> FWebAPIEnumViewModel::Create(const TSharedRef<IWebAPIViewModel>& InParentViewModel, UWebAPIEnum* InEnum, int32 InIndex)
+{
+	TSharedRef<FWebAPIEnumViewModel> ViewModel = MakeShared<FWebAPIEnumViewModel>(FPrivateToken{}, InParentViewModel, InEnum, InIndex);
+	ViewModel->Initialize();
+
+	return ViewModel;
+}
+
+FWebAPIEnumViewModel::FWebAPIEnumViewModel(FPrivateToken, const TSharedRef<IWebAPIViewModel>& InParentViewModel, UWebAPIEnum* InEnum, int32 InIndex)
+	: Schema(InParentViewModel)
+	, Enum(InEnum)
+	, Index(InIndex)
+{
+	check(Enum.IsValid());
+}
+
+void FWebAPIEnumViewModel::Initialize()
+{
+	Values.Empty(Values.Num());
+	for(const TObjectPtr<UWebAPIEnumValue>& Value : Enum->Values)
+	{
+		if(Value)
+		{
+			TSharedRef<FWebAPIEnumValueViewModel> EnumValueViewModel = FWebAPIEnumValueViewModel::Create(AsShared(), Value);
+			Values.Add(MoveTemp(EnumValueViewModel));			
+		}
+	}
+
+	CachedLabel = FText::FromString(Enum->Name.GetDisplayName());
+
+	TArray<FText> TooltipArgs;
+
+	if(Enum->Name.HasTypeInfo() && Enum->Name.TypeInfo->bIsNested)
+	{
+		TooltipArgs.Emplace(LOCTEXT("EnumTooltip_Header", "Enum (Nested)").ToUpper());
+		TooltipArgs.Emplace(FText::Format(LOCTEXT("EnumTooltip_MemberName", "MemberName: {0}"), FText::FromString(Enum->Name.TypeInfo->ToMemberName())));
+	}
+	else
+	{
+		TooltipArgs.Emplace(LOCTEXT("EnumTooltip_Header", "Enum").ToUpper());
+	}
+
+	if(!Enum->Description.IsEmpty())
+	{
+		TooltipArgs.Emplace(FText::Format(LOCTEXT("EnumTooltip_Description", "Description: {0}"), FText::FromString(Enum->Description)));
+	}
+
+	TooltipArgs.Emplace(FText::Format(LOCTEXT("EnumTooltip_JsonName", "JsonName: {0}"), FText::FromString(Enum->Name.GetJsonName())));
+	TooltipArgs.Emplace(FText::Format(LOCTEXT("EnumTooltip_Type", "Type: {0}"), FText::FromString(Enum->Type)));
+
+	const FString DefaultValue = Enum->GetDefaultValue();
+	if(!DefaultValue.IsEmpty())
+	{
+		TooltipArgs.Emplace(FText::Format(LOCTEXT("EnumTooltip_DefaultValue", "Default: {0}"), FText::FromString(DefaultValue)));
+	}
+
+	CachedTooltip = FText::Join(FText::FromString(TEXT("\n")), TooltipArgs);
+}
+
+TSharedPtr<IWebAPIViewModel> FWebAPIEnumViewModel::GetParent()
+{
+	return Schema;
+}
+
+bool FWebAPIEnumViewModel::GetChildren(TArray<TSharedPtr<IWebAPIViewModel>>& OutChildren)
+{
+	OutChildren.Append(Values);	
+	return !Values.IsEmpty();
+}
+
+bool FWebAPIEnumViewModel::GetShouldGenerate() const
+{
+	return Enum.IsValid() && Enum->bGenerate;
+}
+
+#undef LOCTEXT_NAMESPACE
diff --git a/Engine/Plugins/Web/WebAPI/Source/WebAPIEditor/Private/Details/ViewModels/WebAPIEnumViewModel.h b/Engine/Plugins/Web/WebAPI/Source/WebAPIEditor/Private/Details/ViewModels/WebAPIEnumViewModel.h
new file mode 100644
--- /dev/null
+++ b/Engine/Plugins/Web/WebAPI/Source/WebAPIEditor/Private/Details/ViewModels/WebAPIEnumViewModel.h
@@ -0,0 +1,84 @@
+﻿// Copyright Epic Games, Inc. All Rights Reserved.
+
+#pragma once
+
+#include "CoreMinimal.h"
+#include "Details/ViewModels/WebAPIViewModel.h"
+#include "Dom/WebAPIEnum.h"
+
+class FWebAPIEnumViewModel;
+
+class FWebAPIEnumValueViewModel
+	: public TSharedFromThis<FWebAPIEnumValueViewModel>
+	, public IWebAPIViewModel	
+{
+public:
+	static TSharedRef<FWebAPIEnumValueViewModel> Create(const TSharedRef<IWebAPIViewModel>& InParentViewModel, UWebAPIEnumValue* InValue, int32 InIndex = -1);
+
+	virtual ~FWebAPIEnumValueViewModel() = default;
+
+	virtual TSharedPtr<IWebAPIViewModel> GetParent() override;
+	virtual bool GetChildren(TArray<TSharedPtr<IWebAPIViewModel>>& OutChildren) override { return false; }
+	virtual int32 GetIndex() const override { return Index; }
+	virtual const FName& GetViewModelTypeName() override { return IWebAPIViewModel::NAME_EnumValue; }
+
+	const FString& GetValue() const { return Value->Name.GetDisplayName(); }
+
+	/** Checks validity of this ViewModel */
+	virtual bool IsValid() const override { return !Value.IsValid(); }
+
+private:
+	// Private token only allows members or friends to call MakeShared
+	struct FPrivateToken { explicit FPrivateToken() = default; };
+
+public:
+	FWebAPIEnumValueViewModel(FPrivateToken) {}
+	FWebAPIEnumValueViewModel(FPrivateToken, const TSharedRef<IWebAPIViewModel>& InParentViewModel, UWebAPIEnumValue* InValue, int32 InIndex = -1);
+
+private:
+	void Initialize();
+
+private:
+	TSharedPtr<IWebAPIViewModel> Enum;
+	TWeakObjectPtr<UWebAPIEnumValue> Value;
+	int32 Index = 0;
+};
+
+class FWebAPIEnumViewModel
+	: public TSharedFromThis<FWebAPIEnumViewModel>
+	, public IWebAPIViewModel	
+{
+public:
+	static TSharedRef<FWebAPIEnumViewModel> Create(const TSharedRef<IWebAPIViewModel>& InParentViewModel, UWebAPIEnum* InEnum, int32 InIndex = -1);
+	virtual ~FWebAPIEnumViewModel() = default;
+
+	virtual TSharedPtr<IWebAPIViewModel> GetParent() override;
+	virtual bool GetChildren(TArray<TSharedPtr<IWebAPIViewModel>>& OutChildren) override;
+	virtual int32 GetIndex() const override { return Index; }
+	virtual const FName& GetViewModelTypeName() override { return IWebAPIViewModel::NAME_Enum; }
+
+	bool GetShouldGenerate() const;
+	const TArray<TSharedPtr<FWebAPIEnumValueViewModel>>& GetValues() const { return Values; }
+
+	/** Checks validity of this ViewModel */
+	virtual bool IsValid() const override { return Enum.IsValid(); }
+
+private:
+	// Private token only allows members or friends to call MakeShared
+	struct FPrivateToken { explicit FPrivateToken() = default; };
+
+public:
+	FWebAPIEnumViewModel(FPrivateToken) {}
+	FWebAPIEnumViewModel(FPrivateToken, const TSharedRef<IWebAPIViewModel>& InParentViewModel, UWebAPIEnum* InEnum, int32 InIndex = -1);
+
+private:
+	void Initialize();
+
+private:
+	friend class FWebAPIEnumValueViewModel;
+	
+	TSharedPtr<IWebAPIViewModel> Schema;
+	TWeakObjectPtr<UWebAPIEnum> Enum;
+	TArray<TSharedPtr<FWebAPIEnumValueViewModel>> Values;
+	int32 Index = 0;
+};
diff --git a/Engine/Plugins/Web/WebAPI/Source/WebAPIEditor/Private/Details/ViewModels/WebAPIModelViewModel.cpp b/Engine/Plugins/Web/WebAPI/Source/WebAPIEditor/Private/Details/ViewModels/WebAPIModelViewModel.cpp
new file mode 100644
--- /dev/null
+++ b/Engine/Plugins/Web/WebAPI/Source/WebAPIEditor/Private/Details/ViewModels/WebAPIModelViewModel.cpp
@@ -0,0 +1,237 @@
+﻿// Copyright Epic Games, Inc. All Rights Reserved.
+
+#include "WebAPIModelViewModel.h"
+
+#include "Dom/WebAPIEnum.h"
+#include "Dom/WebAPIModel.h"
+#include "Dom/WebAPIType.h"
+
+#define LOCTEXT_NAMESPACE "WebAPIModelViewModel"
+
+TSharedRef<FWebAPIPropertyViewModel> FWebAPIPropertyViewModel::Create(const TSharedRef<IWebAPIViewModel>& InParentViewModel, UWebAPIProperty* InProperty, int32 InIndex)
+{
+	TSharedRef<FWebAPIPropertyViewModel> ViewModel = MakeShared<FWebAPIPropertyViewModel>(FPrivateToken{}, InParentViewModel, InProperty, InIndex);
+	ViewModel->Initialize();
+
+	return ViewModel;
+}
+
+TSharedPtr<IWebAPIViewModel> FWebAPIPropertyViewModel::GetParent()
+{
+	return Model;
+}
+
+bool FWebAPIPropertyViewModel::GetChildren(TArray<TSharedPtr<IWebAPIViewModel>>& OutChildren)
+{
+	if(NestedModel.IsValid())
+	{
+		return NestedModel->GetChildren(OutChildren);
+	}
+	return false;
+}
+
+bool FWebAPIPropertyViewModel::HasCodeText() const
+{
+	if(!IsValid())
+	{
+		return IWebAPIViewModel::HasCodeText();
+	}
+
+	if(Property->Type.HasTypeInfo() && !Property->Type.TypeInfo->Model.IsNull())
+	{
+		if(const TObjectPtr<UWebAPIModel> AsModel = Cast<UWebAPIModel>(Property->Type.TypeInfo->Model.Get()))
+		{
+			return !AsModel->GeneratedCodeText.IsEmpty();
+		}
+		else if(const TObjectPtr<UWebAPIEnum> AsEnum = Cast<UWebAPIEnum>(Property->Type.TypeInfo->Model.Get()))
+		{
+			return !AsEnum->GeneratedCodeText.IsEmpty();
+		}
+	}
+
+	return false;
+}
+
+FText FWebAPIPropertyViewModel::GetCodeText() const
+{
+	if(!IsValid())
+	{
+		return IWebAPIViewModel::GetCodeText();
+	}
+
+	if(Property->Type.HasTypeInfo() && !Property->Type.TypeInfo->Model.IsNull())
+	{
+		if(const TObjectPtr<UWebAPIModel> AsModel = Cast<UWebAPIModel>(Property->Type.TypeInfo->Model.Get()))
+		{
+			return FText::FromString(AsModel->GeneratedCodeText);
+		}
+		else if(const TObjectPtr<UWebAPIEnum> AsEnum = Cast<UWebAPIEnum>(Property->Type.TypeInfo->Model.Get()))
+		{
+			return FText::FromString(AsEnum->GeneratedCodeText);
+		}
+	}
+ 
+	return IWebAPIViewModel::GetCodeText();
+}
+
+FSlateColor FWebAPIPropertyViewModel::GetPinColor() const
+{
+	if(Property.IsValid() && Property->Type.HasTypeInfo())
+	{ 
+		return Property->Type.TypeInfo->PinColor;		
+	}
+
+	return FSlateColor(FLinearColor(0,0,0));
+}
+
+bool FWebAPIPropertyViewModel::IsArray() const
+{ 
+	if(Property.IsValid())
+	{
+		return Property->bIsArray;
+	}
+	return false;
+}
+
+FWebAPIPropertyViewModel::FWebAPIPropertyViewModel(FPrivateToken, const TSharedRef<IWebAPIViewModel>& InParentViewModel, UWebAPIProperty* InProperty, int32 InIndex)
+	: Model(InParentViewModel)
+	, Property(InProperty)
+	, Index(InIndex)
+{
+}
+
+void FWebAPIPropertyViewModel::Initialize()
+{
+	const FString Name = Property->Name.ToString();
+	const FString Description = Property->Description;
+	const FWebAPITypeNameVariant Type = Property->Type;
+
+	if(Type.HasTypeInfo() && Type.TypeInfo->bIsNested && !Type.TypeInfo->Model.IsNull())
+	{
+		NestedModel = UE::WebAPI::Details::CreateViewModel(AsShared(), Type.TypeInfo->Model.Get());
+	}
+
+	FString DebugLabel;
+#if WITH_WEBAPI_DEBUG
+	bool bIsTypeResolved = true;
+	if(Type.HasTypeInfo() && !Type.TypeInfo->bIsBuiltinType)
+	{
+		bIsTypeResolved = !Type.TypeInfo->Model.IsNull();
+	}
+	
+	if(!bIsTypeResolved)
+	{
+		DebugLabel = TEXT(" (Unresolved)");
+	}
+#endif
+
+	CachedLabel = FText::Format(LOCTEXT("PropertyLabel", "{0}: {1}"),
+		FText::FromString(Name),
+		FText::FromString(Type.GetDisplayName() + (IsArray() ? TEXT("[]") : TEXT("")) + DebugLabel));
+
+	TArray<FText> TooltipArgs;
+	TooltipArgs.Emplace(LOCTEXT("PropertyTooltip_Header", "Property").ToUpper());
+	TooltipArgs.Emplace(FText::Format(LOCTEXT("PropertyTooltip_MemberName", "MemberName: {0}"), FText::FromString(Property->Name.ToMemberName())));
+	
+	if(!Description.IsEmpty())
+	{
+		TooltipArgs.Emplace(FText::Format(LOCTEXT("PropertyTooltip_Description", "Description: {0}"), FText::FromString(Description)));
+	}
+
+	TooltipArgs.Emplace(FText::Format(LOCTEXT("PropertyTooltip_JsonName", "JsonName: {0}"), FText::FromString(Property->Name.GetJsonName())));
+	TooltipArgs.Emplace(FText::Format(LOCTEXT("PropertyTooltip_Type", "Type: {0}"),	FText::FromString(Type.ToString())));
+
+	const FString DefaultValue = Property->GetDefaultValue();
+	if(!DefaultValue.IsEmpty())
+	{
+		TooltipArgs.Emplace(FText::Format(LOCTEXT("PropertyTooltip_DefaultValue", "Default: {0}"), FText::FromString(DefaultValue)));
+	}
+
+	CachedTooltip = FText::Join(FText::FromString(TEXT("\n")), TooltipArgs);
+}
+
+TSharedRef<FWebAPIModelViewModel> FWebAPIModelViewModel::Create(const TSharedRef<IWebAPIViewModel>& InParentViewModel, UWebAPIModel* InModel, int32 InIndex)
+{
+	TSharedRef<FWebAPIModelViewModel> ViewModel = MakeShared<FWebAPIModelViewModel>(FPrivateToken{}, InParentViewModel, InModel, InIndex);
+	ViewModel->Initialize();
+
+	return ViewModel;
+}
+
+FWebAPIModelViewModel::FWebAPIModelViewModel(FPrivateToken, const TSharedRef<IWebAPIViewModel>& InParentViewModel, UWebAPIModel* InModel, int32 InIndex)
+	: Schema(InParentViewModel)
+	, Model(InModel)
+	, Index(InIndex)
+{
+	check(Model.IsValid());
+}
+
+void FWebAPIModelViewModel::Initialize()
+{
+	Properties.Empty(Properties.Num());
+	for(const TObjectPtr<UWebAPIProperty>& Property : Model->Properties)
+	{
+		TSharedRef<FWebAPIPropertyViewModel> PropertyViewModel = FWebAPIPropertyViewModel::Create(AsShared(), Property);
+		Properties.Add(MoveTemp(PropertyViewModel));
+	}
+
+	CachedLabel = FText::FromString(Model->Name.GetDisplayName());
+
+	TArray<FText> TooltipArgs;
+	if(Model->Name.HasTypeInfo() && Model->Name.TypeInfo.LoadSynchronous()->bIsNested)
+	{
+		TooltipArgs.Emplace(LOCTEXT("ModelTooltip_Header", "Model (Nested)").ToUpper());
+	}
+	else
+	{
+		TooltipArgs.Emplace(LOCTEXT("ModelTooltip_Header", "Model").ToUpper());
+	}
+
+	TooltipArgs.Emplace(FText::Format(LOCTEXT("ModelTooltip_MemberName", "MemberName: {0}"), FText::FromString(Model->Name.ToMemberName())));
+
+#if WITH_WEBAPI_DEBUG
+	TooltipArgs.Emplace(FText::Format(LOCTEXT("ModelTooltip_FullName", "FullName: {0}"), FText::FromString(Model->Name.ToString())));
+#endif
+	
+	if(!Model->Description.IsEmpty())
+	{
+		TooltipArgs.Emplace(FText::Format(LOCTEXT("ModelTooltip_Description", "Description: {0}"), FText::FromString(Model->Description)));
+	}
+
+	TooltipArgs.Emplace(FText::Format(LOCTEXT("ModelTooltip_JsonName", "JsonName: {0}"), FText::FromString(Model->Name.GetJsonName())));
+	
+	CachedTooltip = FText::Join(FText::FromString(TEXT("\n")), TooltipArgs);
+}
+
+TSharedPtr<IWebAPIViewModel> FWebAPIModelViewModel::GetParent()
+{
+	return Schema;
+}
+
+bool FWebAPIModelViewModel::GetChildren(TArray<TSharedPtr<IWebAPIViewModel>>& OutChildren)
+{
+	OutChildren.Append(Properties);
+	return !Properties.IsEmpty();
+}
+
+bool FWebAPIModelViewModel::HasCodeText() const
+{
+	return IsValid() && !Model->GeneratedCodeText.IsEmpty();
+}
+
+FText FWebAPIModelViewModel::GetCodeText() const
+{
+	if(!IsValid())
+	{
+		return IWebAPIViewModel::GetCodeText();
+	}
+
+	return FText::FromString(Model->GeneratedCodeText);
+}
+
+bool FWebAPIModelViewModel::GetShouldGenerate() const
+{
+	return Model.IsValid() && Model->bGenerate;
+}
+
+#undef LOCTEXT_NAMESPACE
diff --git a/Engine/Plugins/Web/WebAPI/Source/WebAPIEditor/Private/Details/ViewModels/WebAPIModelViewModel.h b/Engine/Plugins/Web/WebAPI/Source/WebAPIEditor/Private/Details/ViewModels/WebAPIModelViewModel.h
new file mode 100644
--- /dev/null
+++ b/Engine/Plugins/Web/WebAPI/Source/WebAPIEditor/Private/Details/ViewModels/WebAPIModelViewModel.h
@@ -0,0 +1,89 @@
+﻿// Copyright Epic Games, Inc. All Rights Reserved.
+
+#pragma once
+
+#include "CoreMinimal.h"
+#include "Details/ViewModels/WebAPIViewModel.h"
+
+class FWebAPIPropertyViewModel
+	: public TSharedFromThis<FWebAPIPropertyViewModel>
+	, public IWebAPIViewModel	
+{
+public:
+	static TSharedRef<FWebAPIPropertyViewModel> Create(const TSharedRef<IWebAPIViewModel>& InParentViewModel, UWebAPIProperty* InProperty, int32 InIndex = -1);
+	virtual ~FWebAPIPropertyViewModel() = default;
+
+	virtual TSharedPtr<IWebAPIViewModel> GetParent() override;
+	virtual bool GetChildren(TArray<TSharedPtr<IWebAPIViewModel>>& OutChildren) override;
+	virtual int32 GetIndex() const override { return Index; }
+	virtual const FName& GetViewModelTypeName() override { return IWebAPIViewModel::NAME_Property; }
+	virtual bool HasCodeText() const override;
+	virtual FText GetCodeText() const override;
+
+	const TSharedPtr<IWebAPIViewModel>& GetModel() const { return Model; }
+	FSlateColor GetPinColor() const;
+
+	bool IsArray() const;
+
+	/** Checks validity of this ViewModel */
+	virtual bool IsValid() const override { return Property.IsValid(); }
+
+private:
+	// Private token only allows members or friends to call MakeShared
+	struct FPrivateToken { explicit FPrivateToken() = default; };
+	
+public:
+	FWebAPIPropertyViewModel(FPrivateToken) {}
+	FWebAPIPropertyViewModel(FPrivateToken, const TSharedRef<IWebAPIViewModel>& InParentViewModel, UWebAPIProperty* InProperty, int32 InIndex = -1);
+
+private:
+	void Initialize();
+
+private:
+	TSharedPtr<IWebAPIViewModel> Model;
+	TWeakObjectPtr<UWebAPIProperty> Property;
+
+	/** Optional nested model to display, if this properties type had to be uniquely generated and has two or more members. */
+	TSharedPtr<IWebAPIViewModel> NestedModel;
+	
+	int32 Index = 0;
+};
+
+class FWebAPIModelViewModel
+	: public TSharedFromThis<FWebAPIModelViewModel>
+	, public IWebAPIViewModel	
+{
+public:
+	static TSharedRef<FWebAPIModelViewModel> Create(const TSharedRef<IWebAPIViewModel>& InParentViewModel, UWebAPIModel* InModel, int32 InIndex = -1);
+	virtual ~FWebAPIModelViewModel() = default;
+
+	virtual TSharedPtr<IWebAPIViewModel> GetParent() override;
+	virtual bool GetChildren(TArray<TSharedPtr<IWebAPIViewModel>>& OutChildren) override;
+	virtual int32 GetIndex() const override { return Index; }
+	virtual const FName& GetViewModelTypeName() override { return IWebAPIViewModel::NAME_Model; }
+	virtual bool HasCodeText() const override;
+	virtual FText GetCodeText() const override;
+
+	bool GetShouldGenerate() const;
+	const TArray<TSharedPtr<FWebAPIPropertyViewModel>>& GetProperties() const { return Properties; }
+
+	/** Checks validity of this ViewModel */
+	virtual bool IsValid() const override { return Model.IsValid(); }
+
+private:
+	// Private token only allows members or friends to call MakeShared
+	struct FPrivateToken { explicit FPrivateToken() = default; };
+
+public:
+	FWebAPIModelViewModel(FPrivateToken) {}
+	FWebAPIModelViewModel(FPrivateToken, const TSharedRef<IWebAPIViewModel>& InParentViewModel, UWebAPIModel* InModel, int32 InIndex = -1);
+
+private:
+	void Initialize();
+
+private:
+	TSharedPtr<IWebAPIViewModel> Schema;
+	TWeakObjectPtr<UWebAPIModel> Model;
+	TArray<TSharedPtr<FWebAPIPropertyViewModel>> Properties;
+	int32 Index = 0;
+};
diff --git a/Engine/Plugins/Web/WebAPI/Source/WebAPIEditor/Private/Details/ViewModels/WebAPIOperationParameterViewModel.cpp b/Engine/Plugins/Web/WebAPI/Source/WebAPIEditor/Private/Details/ViewModels/WebAPIOperationParameterViewModel.cpp
new file mode 100644
--- /dev/null
+++ b/Engine/Plugins/Web/WebAPI/Source/WebAPIEditor/Private/Details/ViewModels/WebAPIOperationParameterViewModel.cpp
@@ -0,0 +1,98 @@
+﻿// Copyright Epic Games, Inc. All Rights Reserved.
+
+#include "WebAPIOperationParameterViewModel.h"
+
+#include "WebAPIOperationRequestViewModel.h"
+#include "Dom/WebAPIOperation.h"
+#include "Dom/WebAPIService.h"
+#include "Dom/WebAPIType.h"
+
+#define LOCTEXT_NAMESPACE "WebAPIOperationParameterViewModel"
+
+TSharedRef<FWebAPIOperationParameterViewModel> FWebAPIOperationParameterViewModel::Create(const TSharedRef<FWebAPIOperationRequestViewModel>& InParentViewModel, UWebAPIOperationParameter* InParameter, int32 InIndex)
+{
+	TSharedRef<FWebAPIOperationParameterViewModel> ViewModel = MakeShared<FWebAPIOperationParameterViewModel>(FPrivateToken{}, InParentViewModel, InParameter, InIndex);
+	ViewModel->Initialize();
+
+	return ViewModel;
+}
+
+TSharedPtr<IWebAPIViewModel> FWebAPIOperationParameterViewModel::GetParent()
+{
+	return Request;
+}
+
+bool FWebAPIOperationParameterViewModel::GetChildren(TArray<TSharedPtr<IWebAPIViewModel>>& OutChildren)
+{
+	return false;
+}
+
+FSlateColor FWebAPIOperationParameterViewModel::GetPinColor() const
+{
+	if(Parameter.IsValid() && Parameter->Type.HasTypeInfo())
+	{ 
+		return Parameter->Type.TypeInfo->PinColor;		
+	}
+
+	return FSlateColor(FLinearColor(0,0,0));
+}
+
+bool FWebAPIOperationParameterViewModel::IsArray() const
+{
+	return Parameter.IsValid() && Parameter->bIsArray;
+}
+
+FWebAPIOperationParameterViewModel::FWebAPIOperationParameterViewModel(FPrivateToken, const TSharedRef<FWebAPIOperationRequestViewModel>& InParentViewModel, UWebAPIOperationParameter* InParameter, int32 InIndex)
+	: Request(InParentViewModel)
+	, Parameter(InParameter)
+	, Index(InIndex)
+{
+}
+
+void FWebAPIOperationParameterViewModel::Initialize()
+{
+	const FWebAPINameVariant Name = Parameter->Name;
+	const FString Description = Parameter->Description;
+	const FWebAPITypeNameVariant Type = Parameter->Type;
+
+	FString DebugLabel;
+#if WITH_WEBAPI_DEBUG
+	bool bIsTypeResolved = true;
+	if(Type.HasTypeInfo() && !Type.TypeInfo->bIsBuiltinType)
+	{
+		bIsTypeResolved = !Type.TypeInfo->Model.IsNull();
+	}
+	
+	if(!bIsTypeResolved)
+	{
+		DebugLabel = TEXT(" (Unresolved)");
+	}
+#endif
+	
+	CachedLabel = FText::Format(LOCTEXT("ParameterLabel", "{0}: {1}"),
+		FText::FromString(Name.ToString()),
+		FText::FromString(Type.GetDisplayName() + (IsArray() ? TEXT("[]") : TEXT("")) + DebugLabel));
+
+	TArray<FText> TooltipArgs;
+	TooltipArgs.Emplace(LOCTEXT("ParameterTooltip_Header", "Property").ToUpper());
+	TooltipArgs.Emplace(FText::Format(LOCTEXT("ParameterTooltip_MemberName", "MemberName: {0}"), FText::FromString(Parameter->Name.ToMemberName())));
+	TooltipArgs.Emplace(FText::Format(LOCTEXT("ParameterTooltip_Storage", "Storage: {0}"), FText::FromString(UE::WebAPI::WebAPIParameterStorage::ToString(Parameter->Storage))));
+
+	if(!Description.IsEmpty())
+	{
+		TooltipArgs.Emplace(FText::Format(LOCTEXT("ParameterTooltip_Description", "Description: {0}"), FText::FromString(Description)));
+	}
+
+	TooltipArgs.Emplace(FText::Format(LOCTEXT("ParameterTooltip_JsonName", "JsonName: {0}"), FText::FromString(Parameter->Name.GetJsonName())));
+	TooltipArgs.Emplace(FText::Format(LOCTEXT("ParameterTooltip_Type", "Type: {0}"),	FText::FromString(Type.ToString())));
+
+	const FString DefaultValue = Parameter->GetDefaultValue();
+	if(!DefaultValue.IsEmpty())
+	{
+		TooltipArgs.Emplace(FText::Format(LOCTEXT("ParameterTooltip_DefaultValue", "Default: {0}"), FText::FromString(DefaultValue)));
+	}
+
+	CachedTooltip = FText::Join(FText::FromString(TEXT("\n")), TooltipArgs);
+}
+
+#undef LOCTEXT_NAMESPACE
diff --git a/Engine/Plugins/Web/WebAPI/Source/WebAPIEditor/Private/Details/ViewModels/WebAPIOperationParameterViewModel.h b/Engine/Plugins/Web/WebAPI/Source/WebAPIEditor/Private/Details/ViewModels/WebAPIOperationParameterViewModel.h
new file mode 100644
--- /dev/null
+++ b/Engine/Plugins/Web/WebAPI/Source/WebAPIEditor/Private/Details/ViewModels/WebAPIOperationParameterViewModel.h
@@ -0,0 +1,51 @@
+﻿// Copyright Epic Games, Inc. All Rights Reserved.
+
+#pragma once
+
+#include "CoreMinimal.h"
+#include "Details/ViewModels/WebAPIViewModel.h"
+#include "Dom/WebAPIOperation.h"
+
+class UWebAPIOperation;
+class UWebAPIService;
+class UWebAPIOperationParameters;
+class FWebAPIOperationRequestViewModel;
+
+/**  */
+class FWebAPIOperationParameterViewModel
+	: public TSharedFromThis<FWebAPIOperationParameterViewModel>
+	, public IWebAPIViewModel
+{
+public:
+	static TSharedRef<FWebAPIOperationParameterViewModel> Create(const TSharedRef<FWebAPIOperationRequestViewModel>& InParentViewModel, UWebAPIOperationParameter* InParameter, int32 InIndex = -1);
+	virtual ~FWebAPIOperationParameterViewModel() = default;
+
+	virtual TSharedPtr<IWebAPIViewModel> GetParent() override;
+	virtual bool GetChildren(TArray<TSharedPtr<IWebAPIViewModel>>& OutChildren) override;
+	virtual int32 GetIndex() const override { return Index; }
+	virtual const FName& GetViewModelTypeName() override { return IWebAPIViewModel::NAME_OperationParameter; }
+
+	const TSharedPtr<FWebAPIOperationRequestViewModel>& GetRequest() const { return Request; }
+	FSlateColor GetPinColor() const;
+
+	bool IsArray() const;
+
+	/** Checks validity of this ViewModel */
+	virtual bool IsValid() const override { return Parameter.IsValid(); }
+
+private:
+	// Private token only allows members or friends to call MakeShared
+	struct FPrivateToken { explicit FPrivateToken() = default; };
+
+public:
+	FWebAPIOperationParameterViewModel(FPrivateToken) {}
+	FWebAPIOperationParameterViewModel(FPrivateToken, const TSharedRef<FWebAPIOperationRequestViewModel>& InParentViewModel, UWebAPIOperationParameter* InParameter, int32 InIndex = -1);
+
+private:
+	void Initialize();
+
+private:
+	TSharedPtr<FWebAPIOperationRequestViewModel> Request;
+	TWeakObjectPtr<UWebAPIOperationParameter> Parameter;
+	int32 Index;
+};
diff --git a/Engine/Plugins/Web/WebAPI/Source/WebAPIEditor/Private/Details/ViewModels/WebAPIOperationRequestViewModel.cpp b/Engine/Plugins/Web/WebAPI/Source/WebAPIEditor/Private/Details/ViewModels/WebAPIOperationRequestViewModel.cpp
new file mode 100644
--- /dev/null
+++ b/Engine/Plugins/Web/WebAPI/Source/WebAPIEditor/Private/Details/ViewModels/WebAPIOperationRequestViewModel.cpp
@@ -0,0 +1,81 @@
+﻿// Copyright Epic Games, Inc. All Rights Reserved.
+
+#include "WebAPIOperationRequestViewModel.h"
+
+#include "WebAPIModelViewModel.h"
+#include "WebAPIOperationParameterViewModel.h"
+#include "WebAPIOperationViewModel.h"
+#include "Algo/AllOf.h"
+#include "Dom/WebAPIOperation.h"
+
+#define LOCTEXT_NAMESPACE "WebAPIOperationRequestViewModel"
+
+TSharedRef<FWebAPIOperationRequestViewModel> FWebAPIOperationRequestViewModel::Create(const TSharedRef<FWebAPIOperationViewModel>& InParentViewModel, UWebAPIOperation* InOperation, int32 InIndex)
+{
+	TSharedRef<FWebAPIOperationRequestViewModel> ViewModel = MakeShared<FWebAPIOperationRequestViewModel>(FPrivateToken{}, InParentViewModel, InOperation, InIndex);
+	ViewModel->Initialize();
+
+	return ViewModel;
+}
+
+TSharedPtr<IWebAPIViewModel> FWebAPIOperationRequestViewModel::GetParent()
+{
+	return Operation;
+}
+
+bool FWebAPIOperationRequestViewModel::GetChildren(TArray<TSharedPtr<IWebAPIViewModel>>& OutChildren)
+{
+	OutChildren.Append(Parameters);
+	if(Body.IsValid())
+	{
+		OutChildren.Add(Body);
+	}
+	return !Parameters.IsEmpty();
+}
+
+bool FWebAPIOperationRequestViewModel::IsValid() const
+{
+	return Algo::AllOf(Parameters, [](const TSharedPtr<FWebAPIOperationParameterViewModel>& InParameter)
+	{
+		return InParameter.IsValid();
+	});
+}
+
+FWebAPIOperationRequestViewModel::FWebAPIOperationRequestViewModel(FPrivateToken, const TSharedRef<FWebAPIOperationViewModel>& InParentViewModel, UWebAPIOperation* InOperation, int32 InIndex)
+	: Operation(InParentViewModel)
+	, OperationObject(InOperation)
+	, Index(InIndex)
+{
+}
+
+void FWebAPIOperationRequestViewModel::Initialize()
+{
+	Parameters.Empty(Parameters.Num());
+
+	for(const TObjectPtr<UWebAPIOperationParameter>& Input : OperationObject->Request->Parameters)
+	{
+		TSharedRef<FWebAPIOperationParameterViewModel> ParameterViewModel = FWebAPIOperationParameterViewModel::Create(AsShared(), Input);
+		Parameters.Add(MoveTemp(ParameterViewModel));
+	}
+
+	CachedLabel = FText::Format(LOCTEXT("RequestLabel", "{0}"), FText::FromString(TEXT("Request")));
+
+	TArray<FText> TooltipArgs;
+
+	const TObjectPtr<UWebAPIOperationRequest> Request = OperationObject->Request;
+
+#if WITH_WEBAPI_DEBUG
+	TooltipArgs.Emplace(FText::Format(LOCTEXT("RequestTooltip_FullName", "FullName: {0}"), FText::FromString(Request->Name.ToString())));
+#endif
+	
+	if(!Request->Description.IsEmpty())
+	{
+		TooltipArgs.Emplace(FText::Format(LOCTEXT("RequestTooltip_Description", "Description: {0}"), FText::FromString(Request->Description)));
+	}
+
+	TooltipArgs.Emplace(FText::Format(LOCTEXT("RequestTooltip_JsonName", "JsonName: {0}"), FText::FromString(Request->Name.GetJsonName())));
+	
+	CachedTooltip = FText::Join(FText::FromString(TEXT("\n")), TooltipArgs);
+}
+
+#undef LOCTEXT_NAMESPACE
diff --git a/Engine/Plugins/Web/WebAPI/Source/WebAPIEditor/Private/Details/ViewModels/WebAPIOperationRequestViewModel.h b/Engine/Plugins/Web/WebAPI/Source/WebAPIEditor/Private/Details/ViewModels/WebAPIOperationRequestViewModel.h
new file mode 100644
--- /dev/null
+++ b/Engine/Plugins/Web/WebAPI/Source/WebAPIEditor/Private/Details/ViewModels/WebAPIOperationRequestViewModel.h
@@ -0,0 +1,52 @@
+﻿// Copyright Epic Games, Inc. All Rights Reserved.
+
+#pragma once
+
+#include "CoreMinimal.h"
+#include "Details/ViewModels/WebAPIViewModel.h"
+#include "Dom/WebAPIOperation.h"
+
+class UWebAPIOperation;
+class UWebAPIService;
+class UWebAPIOperationParameters;
+class FWebAPIOperationViewModel;
+class FWebAPIOperationParameterViewModel;
+
+/** Encapsulates a single Operation Request, containing one or more parameters. */
+class FWebAPIOperationRequestViewModel
+	: public TSharedFromThis<FWebAPIOperationRequestViewModel>
+	, public IWebAPIViewModel
+{
+public:
+	static TSharedRef<FWebAPIOperationRequestViewModel> Create(const TSharedRef<FWebAPIOperationViewModel>& InParentViewModel, UWebAPIOperation* InOperation, int32 InIndex = -1);
+	virtual ~FWebAPIOperationRequestViewModel() = default;
+
+	virtual TSharedPtr<IWebAPIViewModel> GetParent() override;
+	virtual bool GetChildren(TArray<TSharedPtr<IWebAPIViewModel>>& OutChildren) override;
+	virtual int32 GetIndex() const override { return Index; }
+	virtual const FName& GetViewModelTypeName() override { return IWebAPIViewModel::NAME_OperationRequest; }
+	const TArray<TSharedPtr<FWebAPIOperationParameterViewModel>>& GetParameters() const { return Parameters; }
+	bool HasBody() const { return Body.IsValid(); }
+	const TSharedPtr<FWebAPIModelViewModel>& GetBody() const { return Body; }
+
+	/** Checks validity of this ViewModel */
+	virtual bool IsValid() const override;
+
+private:
+	// Private token only allows members or friends to call MakeShared
+	struct FPrivateToken { explicit FPrivateToken() = default; };
+
+public:
+	FWebAPIOperationRequestViewModel(FPrivateToken) {}
+	FWebAPIOperationRequestViewModel(FPrivateToken, const TSharedRef<FWebAPIOperationViewModel>& InParentViewModel, UWebAPIOperation* InOperation, int32 InIndex = -1);
+
+private:
+	void Initialize();
+
+private:
+	TSharedPtr<FWebAPIOperationViewModel> Operation;
+	TWeakObjectPtr<UWebAPIOperation> OperationObject;
+	TArray<TSharedPtr<FWebAPIOperationParameterViewModel>> Parameters;
+	TSharedPtr<FWebAPIModelViewModel> Body; // Re-use Model VM
+	int32 Index;
+};
diff --git a/Engine/Plugins/Web/WebAPI/Source/WebAPIEditor/Private/Details/ViewModels/WebAPIOperationResponseViewModel.cpp b/Engine/Plugins/Web/WebAPI/Source/WebAPIEditor/Private/Details/ViewModels/WebAPIOperationResponseViewModel.cpp
new file mode 100644
--- /dev/null
+++ b/Engine/Plugins/Web/WebAPI/Source/WebAPIEditor/Private/Details/ViewModels/WebAPIOperationResponseViewModel.cpp
@@ -0,0 +1,80 @@
+﻿// Copyright Epic Games, Inc. All Rights Reserved.
+
+#include "WebAPIOperationResponseViewModel.h"
+
+#include "WebAPIOperationViewModel.h"
+#include "Dom/WebAPIOperation.h"
+#include "Dom/WebAPIService.h"
+
+#define LOCTEXT_NAMESPACE "WebAPIOperationResponseViewModel"
+
+TSharedRef<FWebAPIOperationResponseViewModel> FWebAPIOperationResponseViewModel::Create(const TSharedRef<FWebAPIOperationViewModel>& InParentViewModel, UWebAPIOperationResponse* InResponse, int32 InIndex)
+{
+	TSharedRef<FWebAPIOperationResponseViewModel> ViewModel = MakeShared<FWebAPIOperationResponseViewModel>(FPrivateToken{}, InParentViewModel, InResponse, InIndex);
+	ViewModel->Initialize();
+
+	return ViewModel;
+}
+
+TSharedPtr<IWebAPIViewModel> FWebAPIOperationResponseViewModel::GetParent()
+{
+	return Operation;
+}
+
+bool FWebAPIOperationResponseViewModel::GetChildren(TArray<TSharedPtr<IWebAPIViewModel>>& OutChildren)
+{
+	return false;
+}
+
+const FText& FWebAPIOperationResponseViewModel::GetDescription() const
+{
+	return CachedDescription;
+}
+
+FWebAPIOperationResponseViewModel::FWebAPIOperationResponseViewModel(FPrivateToken, const TSharedRef<FWebAPIOperationViewModel>& InParentViewModel, UWebAPIOperationResponse* InResponse, int32 InIndex)
+	: Operation(InParentViewModel)
+	, Response(InResponse)
+	, Index(InIndex)
+{
+}
+
+void FWebAPIOperationResponseViewModel::Initialize()
+{
+	const uint32 ResponseCode = Response->Code;
+	
+	const FWebAPITypeNameVariant Name = Response->Name;
+	const FWebAPITypeNameVariant Type = Response->Type;
+	
+	if(ResponseCode > 0)
+	{
+		CachedLabel = FText::Format(LOCTEXT("ResponseLabel", "Response Code: {0}"),
+			ResponseCode);	
+	}
+	else
+	{
+		CachedLabel = LOCTEXT("ResponseLabel", "Response Code: Default");
+	}
+
+	// Check if response is empty, but has description
+	if(Response->Properties.IsEmpty() && !Response->Description.IsEmpty())
+	{
+		CachedDescription = FText::FromString(Response->Description);
+	}
+
+	TArray<FText> TooltipArgs;
+
+#if WITH_WEBAPI_DEBUG
+	TooltipArgs.Emplace(FText::Format(LOCTEXT("ResponseTooltip_FullName", "FullName: {0}"), FText::FromString(Response->Name.ToString())));
+#endif
+	
+	if(!Response->Description.IsEmpty())
+	{
+		TooltipArgs.Emplace(FText::Format(LOCTEXT("ResponseTooltip_Description", "Description: {0}"), FText::FromString(Response->Description)));
+	}
+
+	TooltipArgs.Emplace(FText::Format(LOCTEXT("ResponseTooltip_JsonName", "JsonName: {0}"), FText::FromString(Response->Name.GetJsonName())));
+	
+	CachedTooltip = FText::Join(FText::FromString(TEXT("\n")), TooltipArgs);
+}
+
+#undef LOCTEXT_NAMESPACE
diff --git a/Engine/Plugins/Web/WebAPI/Source/WebAPIEditor/Private/Details/ViewModels/WebAPIOperationResponseViewModel.h b/Engine/Plugins/Web/WebAPI/Source/WebAPIEditor/Private/Details/ViewModels/WebAPIOperationResponseViewModel.h
new file mode 100644
--- /dev/null
+++ b/Engine/Plugins/Web/WebAPI/Source/WebAPIEditor/Private/Details/ViewModels/WebAPIOperationResponseViewModel.h
@@ -0,0 +1,48 @@
+﻿// Copyright Epic Games, Inc. All Rights Reserved.
+
+#pragma once
+
+#include "CoreMinimal.h"
+#include "Details/ViewModels/WebAPIViewModel.h"
+#include "Dom/WebAPIOperation.h"
+
+class UWebAPIOperation;
+class UWebAPIService;
+class UWebAPIOperationResponse;
+class FWebAPIOperationViewModel;
+
+class FWebAPIOperationResponseViewModel
+	: public TSharedFromThis<FWebAPIOperationResponseViewModel>
+	, public IWebAPIViewModel
+{
+public:
+	static TSharedRef<FWebAPIOperationResponseViewModel> Create(const TSharedRef<FWebAPIOperationViewModel>& InParentViewModel, UWebAPIOperationResponse* InResponse, int32 InIndex = -1);
+	virtual ~FWebAPIOperationResponseViewModel() = default;
+
+	virtual TSharedPtr<IWebAPIViewModel> GetParent() override;
+	virtual bool GetChildren(TArray<TSharedPtr<IWebAPIViewModel>>& OutChildren) override;
+	virtual int32 GetIndex() const override { return Index; }
+	virtual const FName& GetViewModelTypeName() override { return IWebAPIViewModel::NAME_OperationResponse; }
+
+	const FText& GetDescription() const;
+
+	/** Checks validity of this ViewModel */
+	virtual bool IsValid() const override { return Response.IsValid(); }
+
+private:
+	// Private token only allows members or friends to call MakeShared
+	struct FPrivateToken { explicit FPrivateToken() = default; };
+
+public:
+	FWebAPIOperationResponseViewModel(FPrivateToken) {}
+	FWebAPIOperationResponseViewModel(FPrivateToken, const TSharedRef<FWebAPIOperationViewModel>& InParentViewModel, UWebAPIOperationResponse* InResponse, int32 InIndex = -1);
+
+private:
+	void Initialize();
+
+private:
+	TSharedPtr<FWebAPIOperationViewModel> Operation;
+	TWeakObjectPtr<UWebAPIOperationResponse> Response;
+	int32 Index;
+	FText CachedDescription;
+};
diff --git a/Engine/Plugins/Web/WebAPI/Source/WebAPIEditor/Private/Details/ViewModels/WebAPIOperationViewModel.cpp b/Engine/Plugins/Web/WebAPI/Source/WebAPIEditor/Private/Details/ViewModels/WebAPIOperationViewModel.cpp
new file mode 100644
--- /dev/null
+++ b/Engine/Plugins/Web/WebAPI/Source/WebAPIEditor/Private/Details/ViewModels/WebAPIOperationViewModel.cpp
@@ -0,0 +1,174 @@
+﻿// Copyright Epic Games, Inc. All Rights Reserved.
+
+#include "WebAPIOperationViewModel.h"
+
+#include "ScopedTransaction.h"
+#include "WebAPIOperationRequestViewModel.h"
+#include "WebAPIOperationResponseViewModel.h"
+#include "WebAPIServiceViewModel.h"
+#include "Dom/WebAPIOperation.h"
+#include "Misc/StringFormatter.h"
+
+#define LOCTEXT_NAMESPACE "WebAPIOperationViewModel"
+
+TSharedRef<FWebAPIOperationViewModel> FWebAPIOperationViewModel::Create(const TSharedRef<IWebAPIViewModel>& InParentViewModel, UWebAPIOperation* InOperation, int32 InIndex)
+{
+	TSharedRef<FWebAPIOperationViewModel> ViewModel = MakeShared<FWebAPIOperationViewModel>(FPrivateToken{}, InParentViewModel, InOperation, InIndex);
+	ViewModel->Initialize();
+
+	return ViewModel;
+}
+
+TSharedPtr<IWebAPIViewModel> FWebAPIOperationViewModel::GetParent()
+{
+	return Service;
+}
+
+bool FWebAPIOperationViewModel::GetChildren(TArray<TSharedPtr<IWebAPIViewModel>>& OutChildren)
+{
+	OutChildren.Append(Requests);
+	OutChildren.Append(Responses);
+	return !Requests.IsEmpty() || !Responses.IsEmpty();
+}
+
+bool FWebAPIOperationViewModel::HasCodeText() const
+{
+	return IsValid() && !Operation->GeneratedCodeText.IsEmpty();
+}
+
+FText FWebAPIOperationViewModel::GetCodeText() const
+{
+	if(!IsValid())
+	{
+		return IWebAPIViewModel::GetCodeText();
+	}
+
+	return FText::FromString(Operation->GeneratedCodeText);
+}
+
+bool FWebAPIOperationViewModel::GetShouldGenerate(bool bInherit) const
+{
+	if(Operation.IsValid())
+	{
+		bool bParentShouldGenerate = true;
+		if(const TSharedPtr<FWebAPIServiceViewModel> ParentService = StaticCastSharedPtr<FWebAPIServiceViewModel>(Service))
+		{
+			bParentShouldGenerate = ParentService->GetShouldGenerate();
+		}
+		
+		return bInherit
+			? Operation->bGenerate && bParentShouldGenerate
+			: Operation->bGenerate;
+	}
+
+	return false;
+}
+
+void FWebAPIOperationViewModel::SetShouldGenerate(bool bInShouldGenerate) const
+{
+	if(Operation.IsValid())
+	{
+		FScopedTransaction Transaction(LOCTEXT("ModifyOperationGenerate", "Modify operation generate flag"));
+		Operation->Modify();
+		Operation->bGenerate = bInShouldGenerate;
+	}
+}
+
+TSharedPtr<FWebAPIServiceViewModel> FWebAPIOperationViewModel::GetService() const
+{
+	return StaticCastSharedPtr<FWebAPIServiceViewModel>(Service);
+}
+
+FWebAPIOperationViewModel::FWebAPIOperationViewModel(FPrivateToken, const TSharedRef<IWebAPIViewModel>& InParentViewModel, UWebAPIOperation* InOperation, int32 InIndex)
+	: Service(InParentViewModel)
+	, Operation(InOperation)
+	, Index(InIndex)
+{
+}
+
+void FWebAPIOperationViewModel::Initialize()
+{
+	Requests.Empty(Requests.Num());
+	Responses.Empty(Responses.Num());
+
+	// As far as I know, there's only one request model, but allow for multiple
+	TSharedRef<FWebAPIOperationRequestViewModel> RequestViewModel = FWebAPIOperationRequestViewModel::Create(AsShared(), Operation.Get());
+	Requests.Add(MoveTemp(RequestViewModel));
+
+	for(const TObjectPtr<UWebAPIOperationResponse>& Response : Operation->Responses)
+	{
+		TSharedRef<FWebAPIOperationResponseViewModel> ResponseViewModel = FWebAPIOperationResponseViewModel::Create(AsShared(), Response);
+		Responses.Add(MoveTemp(ResponseViewModel));
+	}
+
+	Requests.Sort([](const TSharedPtr<FWebAPIOperationRequestViewModel>& Lhs, const TSharedPtr<FWebAPIOperationRequestViewModel>& Rhs)
+	{
+		return Lhs->GetLabel().ToString() < Rhs->GetLabel().ToString();
+	});
+
+	Responses.Sort([](const TSharedPtr<FWebAPIOperationResponseViewModel>& Lhs, const TSharedPtr<FWebAPIOperationResponseViewModel>& Rhs)
+	{
+		return Lhs->GetLabel().ToString() < Rhs->GetLabel().ToString();
+	});
+
+	// @todo:
+	/*
+	const TArray<TObjectPtr<UWebAPIOperationParameter>> Parameters = Operation->Parameters;
+	TArray<FString> ParameterLabels;
+	Algo::Transform(Parameters, ParameterLabels, [&](const TObjectPtr<UWebAPIOperationParameter>& InParameter)
+	{
+		return FString::Printf(TEXT("%s: %s"), *InParameter->Name.ToString(), *InParameter->Type.ToString());
+	});
+	*/
+
+	/*
+	const FWebAPINameVariant Name = Operation->Name;
+	CachedLabel = FText::Format(LOCTEXT("OperationLabel", "{0}({1}): {2}"),
+		FText::FromString(Name.ToString()),
+		FText::FromString(FString::Join(ParameterLabels, TEXT(", "))),
+		FText::FromString(TEXT("void")));
+		*/
+
+	const FString Name = Operation->Name.GetDisplayName();
+	CachedLabel = FText::Format(LOCTEXT("OperationLabel", "{0}"),
+		FText::FromString(Name));
+
+	FString VerbStr = Operation->Verb;
+	VerbStr.ToUpperInline();
+	VerbLabel = FText::FromString(VerbStr);
+
+	// Add/display path, include tokens ie. "some/path/{id}"
+	const FString PathStr = Operation->Path;
+	PathLabel = FText::FromString(PathStr);
+
+	// RichPathLabel
+	{
+		const FTextFormat TextFormat(PathLabel);
+	
+		TArray<FString> PathFormatArgNames;
+		TextFormat.GetFormatArgumentNames(PathFormatArgNames);
+
+		FFormatNamedArguments PathFormatArgs;
+		for(const FString& ArgName : PathFormatArgNames)
+		{
+			PathFormatArgs.Add(ArgName, FText::FromString(FString::Printf(TEXT("<RichTextBlock.BoldHighlight>{%s}</>"), *ArgName)));			
+		}
+
+		RichPathLabel = FText::Format(TextFormat, PathFormatArgs);
+	}
+
+	TArray<FText> TooltipArgs;
+
+#if WITH_WEBAPI_DEBUG
+	TooltipArgs.Emplace(FText::Format(LOCTEXT("OperationTooltip_FullName", "FullName: {0}"), FText::FromString(Operation->Name.ToString())));
+#endif
+	
+	if(!Operation->Description.IsEmpty())
+	{
+		TooltipArgs.Emplace(FText::Format(LOCTEXT("OperationTooltip_Description", "Description: {0}"), FText::FromString(Operation->Description)));
+	}
+
+	CachedTooltip = FText::Join(FText::FromString(TEXT("\n")), TooltipArgs);
+}
+
+#undef LOCTEXT_NAMESPACE
diff --git a/Engine/Plugins/Web/WebAPI/Source/WebAPIEditor/Private/Details/ViewModels/WebAPIOperationViewModel.h b/Engine/Plugins/Web/WebAPI/Source/WebAPIEditor/Private/Details/ViewModels/WebAPIOperationViewModel.h
new file mode 100644
--- /dev/null
+++ b/Engine/Plugins/Web/WebAPI/Source/WebAPIEditor/Private/Details/ViewModels/WebAPIOperationViewModel.h
@@ -0,0 +1,60 @@
+﻿// Copyright Epic Games, Inc. All Rights Reserved.
+
+#pragma once
+
+#include "CoreMinimal.h"
+#include "Details/ViewModels/WebAPIViewModel.h"
+
+class FWebAPIOperationRequestViewModel;
+class FWebAPIOperationResponseViewModel;
+class UWebAPIOperation;
+class UWebAPIService;
+
+class FWebAPIOperationViewModel
+	: public TSharedFromThis<FWebAPIOperationViewModel>
+	, public IWebAPIViewModel
+{
+public:
+	static TSharedRef<FWebAPIOperationViewModel> Create(const TSharedRef<IWebAPIViewModel>& InParentViewModel, UWebAPIOperation* InOperation, int32 InIndex = -1);
+	virtual ~FWebAPIOperationViewModel() = default;
+
+	virtual TSharedPtr<IWebAPIViewModel> GetParent() override;
+	virtual bool GetChildren(TArray<TSharedPtr<IWebAPIViewModel>>& OutChildren) override;
+	virtual int32 GetIndex() const override { return Index; }
+	virtual const FName& GetViewModelTypeName() override { return IWebAPIViewModel::NAME_Operation; }
+	virtual bool HasCodeText() const override;
+	virtual FText GetCodeText() const override;
+
+	bool GetShouldGenerate(bool bInherit = true) const;
+	void SetShouldGenerate(bool bInShouldGenerate) const;
+	TSharedPtr<FWebAPIServiceViewModel> GetService() const;
+	const FText& GetVerb() const { return VerbLabel; }
+	const FText& GetPath() const { return PathLabel; }
+	const FText& GetRichPath() const { return RichPathLabel; }
+	const TArray<TSharedPtr<FWebAPIOperationRequestViewModel>>& GetRequests() const { return Requests; }
+	const TArray<TSharedPtr<FWebAPIOperationResponseViewModel>>& GetResponses() const { return Responses; }
+
+	/** Checks validity of this ViewModel */
+	virtual bool IsValid() const override { return Operation.IsValid(); }
+
+private:
+	// Private token only allows members or friends to call MakeShared
+	struct FPrivateToken { explicit FPrivateToken() = default; };
+
+public:
+	FWebAPIOperationViewModel(FPrivateToken) {}
+	FWebAPIOperationViewModel(FPrivateToken, const TSharedRef<IWebAPIViewModel>& InParentViewModel, UWebAPIOperation* InOperation, int32 InIndex = -1);
+
+private:
+	void Initialize();
+
+private:
+	TSharedPtr<IWebAPIViewModel> Service;
+	TWeakObjectPtr<UWebAPIOperation> Operation;
+	TArray<TSharedPtr<FWebAPIOperationRequestViewModel>> Requests;
+	TArray<TSharedPtr<FWebAPIOperationResponseViewModel>> Responses;
+	int32 Index;
+	FText VerbLabel;
+	FText PathLabel;
+	FText RichPathLabel;
+};
diff --git a/Engine/Plugins/Web/WebAPI/Source/WebAPIEditor/Private/Details/ViewModels/WebAPIServiceViewModel.cpp b/Engine/Plugins/Web/WebAPI/Source/WebAPIEditor/Private/Details/ViewModels/WebAPIServiceViewModel.cpp
new file mode 100644
--- /dev/null
+++ b/Engine/Plugins/Web/WebAPI/Source/WebAPIEditor/Private/Details/ViewModels/WebAPIServiceViewModel.cpp
@@ -0,0 +1,70 @@
+﻿// Copyright Epic Games, Inc. All Rights Reserved.
+
+#include "WebAPIServiceViewModel.h"
+
+#include "ScopedTransaction.h"
+#include "WebAPIOperationViewModel.h"
+#include "Dom/WebAPIOperation.h"
+#include "Dom/WebAPIService.h"
+
+#define LOCTEXT_NAMESPACE "WebAPIServiceViewModel"
+
+TSharedRef<FWebAPIServiceViewModel> FWebAPIServiceViewModel::Create(const TSharedRef<IWebAPIViewModel>& InParentViewModel, UWebAPIService* InService, int32 InIndex)
+{
+	TSharedRef<FWebAPIServiceViewModel> ViewModel = MakeShared<FWebAPIServiceViewModel>(FPrivateToken{}, InParentViewModel, InService, InIndex);
+	ViewModel->Initialize();
+
+	return ViewModel;
+}
+
+bool FWebAPIServiceViewModel::GetChildren(TArray<TSharedPtr<IWebAPIViewModel>>& OutChildren)
+{
+	OutChildren.Append(Operations);
+	return !Operations.IsEmpty();
+}
+
+bool FWebAPIServiceViewModel::GetShouldGenerate() const
+{
+	if(Service.IsValid())
+	{
+		return Service->bGenerate;
+	}
+
+	return false;
+}
+
+void FWebAPIServiceViewModel::SetShouldGenerate(bool bInShouldGenerate) const
+{
+	if(Service.IsValid())
+	{
+		FScopedTransaction Transaction(LOCTEXT("ModifyServiceGenerate", "Modify service generate flag"));
+		Service->Modify();
+		Service->bGenerate = bInShouldGenerate;
+	}
+}
+
+FWebAPIServiceViewModel::FWebAPIServiceViewModel(FPrivateToken, const TSharedRef<IWebAPIViewModel>& InParentViewModel, UWebAPIService* InService, int32 InIndex)
+	: Service(InService)
+	, Index(InIndex)
+{
+	check(Service.IsValid());
+}
+
+void FWebAPIServiceViewModel::Initialize()
+{
+	Operations.Empty(Operations.Num());
+	for(TObjectPtr<UWebAPIOperation>& Operation : Service->Operations)
+	{
+		TSharedRef<FWebAPIOperationViewModel> OperationViewModel = FWebAPIOperationViewModel::Create(AsShared(), Operation);
+		Operations.Add(MoveTemp(OperationViewModel));
+	}
+
+	Operations.Sort([](const TSharedPtr<FWebAPIOperationViewModel>& Lhs, const TSharedPtr<FWebAPIOperationViewModel>& Rhs)
+	{
+		return Lhs->GetLabel().ToString() < Rhs->GetLabel().ToString();
+	});
+
+	CachedLabel = FText::FromString(Service->Name.GetDisplayName());
+}
+
+#undef LOCTEXT_NAMESPACE
diff --git a/Engine/Plugins/Web/WebAPI/Source/WebAPIEditor/Private/Details/ViewModels/WebAPIServiceViewModel.h b/Engine/Plugins/Web/WebAPI/Source/WebAPIEditor/Private/Details/ViewModels/WebAPIServiceViewModel.h
new file mode 100644
--- /dev/null
+++ b/Engine/Plugins/Web/WebAPI/Source/WebAPIEditor/Private/Details/ViewModels/WebAPIServiceViewModel.h
@@ -0,0 +1,47 @@
+﻿// Copyright Epic Games, Inc. All Rights Reserved.
+
+#pragma once
+
+#include "CoreMinimal.h"
+#include "Dom/WebAPISchema.h"
+#include "Details/ViewModels/WebAPIViewModel.h"
+
+class UWebAPIService;
+class FWebAPIOperationViewModel;
+
+class FWebAPIServiceViewModel
+	: public TSharedFromThis<FWebAPIServiceViewModel>
+	, public IWebAPIViewModel	
+{
+public:
+	static TSharedRef<FWebAPIServiceViewModel> Create(const TSharedRef<IWebAPIViewModel>& InParentViewModel, UWebAPIService* InService, int32 InIndex = -1);
+	virtual ~FWebAPIServiceViewModel() = default;
+
+	virtual TSharedPtr<IWebAPIViewModel> GetParent() override { return nullptr; }
+	virtual bool GetChildren(TArray<TSharedPtr<IWebAPIViewModel>>& OutChildren) override;
+	virtual int32 GetIndex() const override { return Index; }
+	virtual const FName& GetViewModelTypeName() override { return IWebAPIViewModel::NAME_Service; }
+
+	bool GetShouldGenerate() const;
+	void SetShouldGenerate(bool bInShouldGenerate) const;
+	const TArray<TSharedPtr<FWebAPIOperationViewModel>>& GetOperations() const { return Operations; }
+
+	/** Checks validity of this ViewModel */
+	virtual bool IsValid() const override { return Service.IsValid(); }
+
+private:
+	// Private token only allows members or friends to call MakeShared
+	struct FPrivateToken { explicit FPrivateToken() = default; };
+
+public:
+	FWebAPIServiceViewModel(FPrivateToken) {}
+	FWebAPIServiceViewModel(FPrivateToken, const TSharedRef<IWebAPIViewModel>& InParentViewModel, UWebAPIService* InService, int32 InIndex = -1);
+
+private:
+	void Initialize();
+
+private:
+	TWeakObjectPtr<UWebAPIService> Service;
+	TArray<TSharedPtr<FWebAPIOperationViewModel>> Operations;
+	int32 Index;
+};
diff --git a/Engine/Plugins/Web/WebAPI/Source/WebAPIEditor/Private/Details/ViewModels/WebAPIViewModel.cpp b/Engine/Plugins/Web/WebAPI/Source/WebAPIEditor/Private/Details/ViewModels/WebAPIViewModel.cpp
new file mode 100644
--- /dev/null
+++ b/Engine/Plugins/Web/WebAPI/Source/WebAPIEditor/Private/Details/ViewModels/WebAPIViewModel.cpp
@@ -0,0 +1,209 @@
+﻿// Copyright Epic Games, Inc. All Rights Reserved.
+
+#include "WebAPIViewModel.h"
+
+#include "WebAPIModelViewModel.h"
+#include "WebAPIServiceViewModel.h"
+#include "WebAPIEnumViewModel.h"
+#include "WebAPIOperationViewModel.h"
+#include "WebAPIParameterViewModel.h"
+#include "Dom/WebAPIModel.h"
+#include "Dom/WebAPIParameter.h"
+#include "Dom/WebAPIService.h"
+#include "Dom/WebAPIType.h"
+
+#define LOCTEXT_NAMESPACE "WebAPIViewModel"
+
+FName IWebAPIViewModel::NAME_Definition = TEXT("Definition");
+FName IWebAPIViewModel::NAME_Schema = TEXT("Schema");
+FName IWebAPIViewModel::NAME_Model = TEXT("Model");
+FName IWebAPIViewModel::NAME_Property = TEXT("Property");
+FName IWebAPIViewModel::NAME_Enum = TEXT("Enum");
+FName IWebAPIViewModel::NAME_EnumValue = TEXT("EnumValue");
+FName IWebAPIViewModel::NAME_Parameter = TEXT("ParameterModel");
+FName IWebAPIViewModel::NAME_Service = TEXT("Service");
+FName IWebAPIViewModel::NAME_Operation = TEXT("Operation");
+FName IWebAPIViewModel::NAME_OperationRequest = TEXT("OperationRequest");
+FName IWebAPIViewModel::NAME_OperationParameter = TEXT("OperationParameter");
+FName IWebAPIViewModel::NAME_OperationResponse = TEXT("OperationResponse");
+FName IWebAPIViewModel::NAME_Code = TEXT("Code");
+FName IWebAPIViewModel::NAME_MessageLog = TEXT("MessageLog");
+
+template <typename ModelType, class ParentViewModelType, class ViewModelType>
+TSharedPtr<ViewModelType> UE::WebAPI::Details::CreateViewModel(const TSharedRef<ParentViewModelType>& InParentViewModel, ModelType* InModel)
+{
+	if(!InModel)
+	{
+		return nullptr;
+	}
+	
+	const UClass* ModelClass = InModel->GetClass();
+	if(ModelClass == UWebAPIEnum::StaticClass())
+	{
+		return FWebAPIEnumViewModel::Create(InParentViewModel, Cast<UWebAPIEnum>(InModel));
+	}
+	else if(ModelClass == UWebAPIEnumValue::StaticClass())
+	{
+		return FWebAPIEnumValueViewModel::Create(InParentViewModel, Cast<UWebAPIEnumValue>(InModel));
+	}
+	else if(ModelClass == UWebAPIModel::StaticClass())
+	{
+		return FWebAPIModelViewModel::Create(InParentViewModel, Cast<UWebAPIModel>(InModel));
+	}
+	else if(ModelClass == UWebAPIProperty::StaticClass())
+	{
+		return FWebAPIPropertyViewModel::Create(InParentViewModel, Cast<UWebAPIProperty>(InModel));
+	}
+	else if(ModelClass == UWebAPIService::StaticClass())
+	{
+		return FWebAPIServiceViewModel::Create(InParentViewModel, Cast<UWebAPIService>(InModel));
+	}
+	else if(ModelClass == UWebAPIOperation::StaticClass())
+	{
+		return FWebAPIOperationViewModel::Create(InParentViewModel, Cast<UWebAPIOperation>(InModel));
+	}
+	else if(ModelClass == UWebAPIParameter::StaticClass())
+	{
+		return FWebAPIParameterViewModel::Create(InParentViewModel, Cast<UWebAPIParameter>(InModel));
+	}
+	else
+	{
+		checkNoEntry();
+		return nullptr;
+	}
+}
+
+FText IWebAPIViewModel::GetCodeText() const
+{
+	return FText::GetEmpty();
+}
+
+TSharedRef<FWebAPIDefinitionViewModel> FWebAPIDefinitionViewModel::Create(UWebAPIDefinition* InDefinition)
+{
+	TSharedRef<FWebAPIDefinitionViewModel> ViewModel = MakeShared<FWebAPIDefinitionViewModel>(FPrivateToken{}, InDefinition);
+	ViewModel->Initialize();
+
+	return ViewModel;
+}
+
+FWebAPIDefinitionViewModel::FWebAPIDefinitionViewModel(FPrivateToken, UWebAPIDefinition* InDefinition)
+	: Definition(InDefinition)
+{
+	check(Definition.IsValid());
+}
+
+bool FWebAPIDefinitionViewModel::GetChildren(TArray<TSharedPtr<IWebAPIViewModel>>& OutChildren)
+{
+	OutChildren.Add(Schema);	
+	return true;
+}
+
+void FWebAPIDefinitionViewModel::Refresh()
+{
+	Schema->Refresh();
+}
+
+bool FWebAPIDefinitionViewModel::IsSameDefinition(UObject* InObject) const
+{
+	if(InObject && IsValid())
+	{
+		return InObject == Definition.Get();
+	}
+
+	return false;
+}
+
+bool FWebAPIDefinitionViewModel::IsValid() const
+{
+	return Definition.IsValid();
+}
+
+void FWebAPIDefinitionViewModel::Initialize()
+{
+	Schema = FWebAPISchemaViewModel::Create(AsShared(), Definition->GetWebAPISchema());
+	CachedLabel = FText::GetEmpty();
+}
+
+TSharedRef<FWebAPISchemaViewModel> FWebAPISchemaViewModel::Create(const TSharedRef<FWebAPIDefinitionViewModel>& InParentViewModel, UWebAPISchema* InSchema)
+{
+	TSharedRef<FWebAPISchemaViewModel> ViewModel = MakeShared<FWebAPISchemaViewModel>(FPrivateToken{}, InParentViewModel, InSchema);
+	ViewModel->Initialize();
+
+	return ViewModel;
+}
+
+void FWebAPISchemaViewModel::Refresh()
+{
+	Initialize();
+}
+
+FWebAPISchemaViewModel::FWebAPISchemaViewModel(FPrivateToken, const TSharedRef<FWebAPIDefinitionViewModel>& InParentViewModel, UWebAPISchema* InSchema)
+	: Definition(InParentViewModel)
+	, Schema(InSchema)
+{
+	check(Schema.IsValid());
+}
+
+void FWebAPISchemaViewModel::Initialize()
+{
+	Services.Empty(Services.Num());
+	for(TTuple<FString, TObjectPtr<UWebAPIService>>& Service : Schema->Services)
+	{
+		TSharedRef<FWebAPIServiceViewModel> ServiceViewModel = FWebAPIServiceViewModel::Create(AsShared(), Service.Value);
+		Services.Add(MoveTemp(ServiceViewModel));
+	}
+
+	TArray<TObjectPtr<UWebAPIModelBase>> SourceModels = Schema->Models;
+	Algo::SortBy(SourceModels, [](const TObjectPtr<UWebAPIModelBase>& InModel)
+	{
+		return InModel->GetSortKey();		
+	});
+
+	Models.Empty(Models.Num());
+	for(const TObjectPtr<UWebAPIModelBase>& Model : SourceModels)
+	{
+		if (UWebAPIEnum* SrcEnum = Cast<UWebAPIEnum>(Model))
+		{
+			if(SrcEnum->Name.HasTypeInfo() && SrcEnum->Name.TypeInfo->bIsNested)
+			{
+				continue;
+			}
+			TSharedRef<FWebAPIEnumViewModel> EnumViewModel = FWebAPIEnumViewModel::Create(AsShared(), SrcEnum);
+			Models.Add(MoveTemp(EnumViewModel));
+		}
+		else if (UWebAPIParameter* SrcParameter = Cast<UWebAPIParameter>(Model))
+		{
+			if(SrcParameter->Name.HasTypeInfo() && SrcParameter->Name.TypeInfo->bIsNested)
+			{
+				continue;	
+			}			
+			TSharedRef<FWebAPIParameterViewModel> ParameterViewModel = FWebAPIParameterViewModel::Create(AsShared(), SrcParameter);
+			Models.Add(MoveTemp(ParameterViewModel));
+		}
+		else if (UWebAPIModel* SrcModel = Cast<UWebAPIModel>(Model))
+		{
+			if(SrcModel->Name.HasTypeInfo() && SrcModel->Name.TypeInfo->bIsNested)
+			{
+				continue;
+			}
+			TSharedRef<FWebAPIModelViewModel> ModelViewModel = FWebAPIModelViewModel::Create(AsShared(), SrcModel);
+			Models.Add(MoveTemp(ModelViewModel));
+		}
+	}
+
+	CachedLabel = FText::GetEmpty();
+}
+
+TSharedPtr<IWebAPIViewModel> FWebAPISchemaViewModel::GetParent()
+{
+	return Definition;
+}
+
+bool FWebAPISchemaViewModel::GetChildren(TArray<TSharedPtr<IWebAPIViewModel>>& OutChildren)
+{
+	OutChildren.Append(Services);
+	OutChildren.Append(Models);
+	return !Services.IsEmpty() && !Models.IsEmpty();
+}
+
+#undef LOCTEXT_NAMESPACE
diff --git a/Engine/Plugins/Web/WebAPI/Source/WebAPIEditor/Private/Details/ViewModels/WebAPIViewModel.h b/Engine/Plugins/Web/WebAPI/Source/WebAPIEditor/Private/Details/ViewModels/WebAPIViewModel.h
new file mode 100644
--- /dev/null
+++ b/Engine/Plugins/Web/WebAPI/Source/WebAPIEditor/Private/Details/ViewModels/WebAPIViewModel.h
@@ -0,0 +1,191 @@
+﻿// Copyright Epic Games, Inc. All Rights Reserved.
+
+#pragma once
+
+#include "CoreMinimal.h"
+#include "Dom/WebAPISchema.h"
+#include "WebAPIDefinition.h"
+
+class FWebAPIModelViewModel;
+class FWebAPIDefinitionViewModel;
+class FWebAPISchemaViewModel;
+class FWebAPIServiceViewModel;
+class IWebAPIViewModel;
+
+namespace UE
+{
+	namespace WebAPI
+	{
+		namespace Details
+		{
+			template <typename ModelType, class ParentViewModelType = IWebAPIViewModel, class ViewModelType = IWebAPIViewModel> 
+			TSharedPtr<ViewModelType> CreateViewModel(const TSharedRef<ParentViewModelType>& InParentViewModel, ModelType* InModel);
+		}
+	}
+}
+
+/** Common interface for View Models. */
+class IWebAPIViewModel
+{
+protected:
+	~IWebAPIViewModel() = default;
+
+public:
+	/** Get Parent ViewModel, if any. */
+	virtual TSharedPtr<class IWebAPIViewModel> GetParent() = 0;
+
+	/** Get Parent ViewModel of ModelType, if any. */
+	template <typename ModelType>
+	TSharedPtr<ModelType> GetParent();
+
+	/** Get Label Text. */
+	virtual FText GetLabel() { return CachedLabel; }
+
+	/** Get Tooltip Text. */
+	virtual FText GetTooltip() { return CachedTooltip; }
+	
+	/** See STreeView::OnGetChildren, return false if no children. */
+	virtual bool GetChildren(TArray<TSharedPtr<IWebAPIViewModel>>& OutChildren) = 0;
+
+	/** Get the index of this element, if applicable. */
+	virtual int32 GetIndex() const = 0;
+
+	/** Get the ViewModel type name. */
+	virtual const FName& GetViewModelTypeName() = 0;
+
+	/** Whether this ViewModel should be initially expanded in a treeview. */
+	virtual bool ShouldExpandByDefault() { return true; }
+
+	/** Rebuild the ViewModel. */
+	virtual void Refresh() { };
+
+	/** Checks validity of this ViewModel */
+	virtual bool IsValid() const = 0;
+
+	/** Check if this viewmodel context contains any code to display. */
+	virtual bool HasCodeText() const { return false; }
+
+	/** Get code text, if present. */
+	virtual FText GetCodeText() const;
+
+public:
+	static FName NAME_Definition;
+	static FName NAME_Schema;
+	static FName NAME_Model;	
+	static FName NAME_Property;
+	static FName NAME_Enum;
+	static FName NAME_EnumValue;
+	static FName NAME_Parameter;
+	static FName NAME_Service;
+	static FName NAME_Operation;
+	static FName NAME_OperationRequest;
+	static FName NAME_OperationParameter;
+	static FName NAME_OperationResponse;
+	static FName NAME_Code;
+	static FName NAME_MessageLog;
+
+protected:
+	FText CachedLabel;
+	FText CachedTooltip;
+};
+
+template <class ModelType>
+TSharedPtr<ModelType> IWebAPIViewModel::GetParent()
+{
+	TSharedPtr<IWebAPIViewModel> Parent = GetParent();
+	while(Parent.IsValid())
+	{
+		if(const TSharedPtr<ModelType> FoundParentOfType = StaticCastSharedPtr<ModelType>(Parent))
+		{
+			return FoundParentOfType;			
+		}
+
+		// Step up the tree
+		Parent = Parent->GetParent();
+	}
+
+	return nullptr;
+}
+
+class FWebAPIDefinitionViewModel
+	: public TSharedFromThis<FWebAPIDefinitionViewModel>
+	, public IWebAPIViewModel
+{
+public:
+	static TSharedRef<FWebAPIDefinitionViewModel> Create(UWebAPIDefinition* InDefinition);
+	virtual ~FWebAPIDefinitionViewModel() = default;
+
+	virtual TSharedPtr<IWebAPIViewModel> GetParent() override { return nullptr; }
+	virtual bool GetChildren(TArray<TSharedPtr<IWebAPIViewModel>>& OutChildren) override;
+	virtual int32 GetIndex() const override { return 0; }
+	virtual const FName& GetViewModelTypeName() override { return IWebAPIViewModel::NAME_Definition; }
+
+	const TWeakObjectPtr<UWebAPIDefinition>& GetDefinition() const	{ return Definition; }
+	TSharedPtr<FWebAPISchemaViewModel> GetSchema() const { return Schema; }
+
+	/** Rebuild the ViewModel. */
+	virtual void Refresh() override;;
+
+	bool IsSameDefinition(UObject* InObject) const;
+
+	/** Checks validity of this ViewModel */
+	virtual bool IsValid() const override;
+
+private:
+	// Private token only allows members or friends to call MakeShared
+	struct FPrivateToken { explicit FPrivateToken() = default; };
+
+public:
+	FWebAPIDefinitionViewModel(FPrivateToken) {}
+	FWebAPIDefinitionViewModel(FPrivateToken, UWebAPIDefinition* InDefinition);
+
+private:
+	void Initialize();
+
+private:
+	TWeakObjectPtr<UWebAPIDefinition> Definition;
+	TSharedPtr<FWebAPISchemaViewModel> Schema;
+};
+
+class FWebAPISchemaViewModel
+	: public TSharedFromThis<FWebAPISchemaViewModel>
+	, public IWebAPIViewModel
+{
+public:
+	static TSharedRef<FWebAPISchemaViewModel> Create(const TSharedRef<FWebAPIDefinitionViewModel>& InParentViewModel, UWebAPISchema* InSchema);
+	virtual ~FWebAPISchemaViewModel() = default;
+
+	virtual TSharedPtr<IWebAPIViewModel> GetParent() override;
+	virtual bool GetChildren(TArray<TSharedPtr<IWebAPIViewModel>>& OutChildren) override;
+	virtual int32 GetIndex() const override { return 0; }
+	virtual const FName& GetViewModelTypeName() override { return IWebAPIViewModel::NAME_Schema; }
+	
+	const TArray<TSharedPtr<FWebAPIServiceViewModel>>& GetServices() const { return Services; }
+	const TArray<TSharedPtr<IWebAPIViewModel>>& GetModels() const { return Models; }
+
+	/** Rebuild the ViewModel. */
+	virtual void Refresh() override;
+
+	/** Whether this ViewModel should be initially expanded. */
+	virtual bool ShouldExpandByDefault() override { return true; }
+
+	/** Checks validity of this ViewModel */
+	virtual bool IsValid() const override { return Schema.IsValid(); }
+
+private:
+	// Private token only allows members or friends to call MakeShared
+	struct FPrivateToken { explicit FPrivateToken() = default; };
+
+public:
+	FWebAPISchemaViewModel(FPrivateToken) {}
+	FWebAPISchemaViewModel(FPrivateToken, const TSharedRef<FWebAPIDefinitionViewModel>& InParentViewModel, UWebAPISchema* InSchema);
+
+private:
+	void Initialize();
+
+private:
+	TSharedPtr<FWebAPIDefinitionViewModel> Definition;
+	TWeakObjectPtr<UWebAPISchema> Schema;
+	TArray<TSharedPtr<FWebAPIServiceViewModel>> Services;
+	TArray<TSharedPtr<IWebAPIViewModel>> Models;
+};
diff --git a/Engine/Plugins/Web/WebAPI/Source/WebAPIEditor/Private/Details/Widgets/SWebAPIMessageLog.cpp b/Engine/Plugins/Web/WebAPI/Source/WebAPIEditor/Private/Details/Widgets/SWebAPIMessageLog.cpp
new file mode 100644
--- /dev/null
+++ b/Engine/Plugins/Web/WebAPI/Source/WebAPIEditor/Private/Details/Widgets/SWebAPIMessageLog.cpp
@@ -0,0 +1,25 @@
+﻿// Copyright Epic Games, Inc. All Rights Reserved.
+
+#include "SWebAPIMessageLog.h"
+
+#include "MessageLogModule.h"
+#include "WebAPIMessageLog.h"
+#include "Modules/ModuleManager.h"
+#include "Widgets/Layout/SBorder.h"
+
+void SWebAPIMessageLog::Construct(const FArguments& InArgs, const TSharedRef<FWebAPIMessageLog>& InMessageLog)
+{
+	// Create widget from MessageLog module
+	FMessageLogModule& MessageLogModule = FModuleManager::LoadModuleChecked<FMessageLogModule>("MessageLog");
+	LogListingWidget = MessageLogModule.CreateLogListingWidget(InMessageLog->GetMessageLogListing().ToSharedRef());
+
+	ChildSlot
+	[
+		SNew(SBorder)
+		.BorderImage(FAppStyle::Get().GetBrush("Menu.Background"))
+		.Padding(2.f)
+		[
+			LogListingWidget.ToSharedRef()
+		]
+	];
+}
diff --git a/Engine/Plugins/Web/WebAPI/Source/WebAPIEditor/Private/Details/Widgets/SWebAPIMessageLog.h b/Engine/Plugins/Web/WebAPI/Source/WebAPIEditor/Private/Details/Widgets/SWebAPIMessageLog.h
new file mode 100644
--- /dev/null
+++ b/Engine/Plugins/Web/WebAPI/Source/WebAPIEditor/Private/Details/Widgets/SWebAPIMessageLog.h
@@ -0,0 +1,27 @@
+﻿// Copyright Epic Games, Inc. All Rights Reserved.
+
+#pragma once
+
+#include "CoreMinimal.h"
+#include "WebAPIMessageLog.h"
+#include "Widgets/SCompoundWidget.h"
+
+class SWidget;
+
+/** Lists WebAPI log entries. */
+class WEBAPIEDITOR_API SWebAPIMessageLog : public SCompoundWidget
+{
+public:
+	SLATE_BEGIN_ARGS(SWebAPIMessageLog)
+	{ }
+	SLATE_END_ARGS()
+
+	/** Constructs this widget with InArgs */
+	void Construct(const FArguments& InArgs, const TSharedRef<FWebAPIMessageLog>& InMessageLog);
+
+private:
+	/** Pointer to log widget */
+	TSharedPtr<SWidget> LogListingWidget;
+
+	TSharedPtr<FWebAPIMessageLog> MessageLog;
+};
diff --git a/Engine/Plugins/Web/WebAPI/Source/WebAPIEditor/Private/Details/Widgets/SWebAPIModulePicker.cpp b/Engine/Plugins/Web/WebAPI/Source/WebAPIEditor/Private/Details/Widgets/SWebAPIModulePicker.cpp
new file mode 100644
--- /dev/null
+++ b/Engine/Plugins/Web/WebAPI/Source/WebAPIEditor/Private/Details/Widgets/SWebAPIModulePicker.cpp
@@ -0,0 +1,159 @@
+﻿// Copyright Epic Games, Inc. All Rights Reserved.
+
+#include "SWebAPIModulePicker.h"
+
+#include "GameProjectUtils.h"
+#include "IDocumentation.h"
+#include "WebAPIEditorSubsystem.h"
+#include "Misc/App.h"
+#include "Widgets/Input/SHyperlink.h"
+#include "Widgets/Layout/SGridPanel.h"
+
+#define LOCTEXT_NAMESPACE "WebAPIModulePicker"
+
+/** The last selected module name. Meant to keep the same module selected after first selection. */
+FString SWebAPIModulePicker::LastSelectedModuleName;
+
+void SWebAPIModulePicker::Construct(const FArguments& InArgs)
+{
+	ModuleName = InArgs._ModuleName;
+	OnModuleChanged = InArgs._OnModuleChanged;
+	
+	GEditor->GetEditorSubsystem<UWebAPIEditorSubsystem>()->OnModuleCreated().AddSP(this, &SWebAPIModulePicker::OnModuleCreated);	
+	Refresh();
+	
+	ChildSlot
+	[
+		SNew(SHorizontalBox)
+
+		+ SHorizontalBox::Slot()
+		.AutoWidth()
+		.Padding(0.0f, 0.0f, 0.0f, 0.0f)
+		[
+			SAssignNew(AvailableModulesCombo, SComboBox<TSharedPtr<FModuleContextInfo>>)
+			.ToolTipText( LOCTEXT("ModuleComboToolTip", "Choose the target module for your generated files") )
+			.OptionsSource( &AvailableModules )
+			.InitiallySelectedItem( SelectedModuleInfo )
+			.OnSelectionChanged( this, &SWebAPIModulePicker::SelectedModuleComboBoxSelectionChanged )
+			.OnGenerateWidget( this, &SWebAPIModulePicker::MakeWidgetForSelectedModuleCombo )
+			[
+				SNew(STextBlock)
+				.Text( this, &SWebAPIModulePicker::GetSelectedModuleComboText )
+			]
+		]
+	];
+}
+
+void SWebAPIModulePicker::Refresh()
+{
+	{
+		TArray<FModuleContextInfo> CurrentModules = GameProjectUtils::GetCurrentProjectModules();
+		check(!CurrentModules.IsEmpty());
+
+		const TArray<FModuleContextInfo> CurrentPluginModules = GameProjectUtils::GetCurrentProjectPluginModules();
+
+		CurrentModules.Append(CurrentPluginModules);
+
+		AvailableModules.Reserve(CurrentModules.Num());
+		for(FModuleContextInfo& ModuleInfo : CurrentModules)
+		{
+			AvailableModules.Emplace(MakeShared<FModuleContextInfo>(MoveTemp(ModuleInfo)));
+		}
+
+		Algo::SortBy(AvailableModules, &FModuleContextInfo::ModuleName);
+	}
+
+	// If we've been given an initial path that maps to a valid project module, use that as our initial module and path
+	for(const auto& AvailableModule : AvailableModules)
+	{
+		if(AvailableModule->ModuleName.Equals(ModuleName, ESearchCase::IgnoreCase))
+		{
+			SelectedModuleInfo = AvailableModule;
+			break;
+		}
+	}
+
+	// If we didn't get given a valid path override (see above), try and automatically work out the best default module
+	// If we have a runtime module with the same name as our project, then use that
+	// Otherwise, set out default target module as the first runtime module in the list
+	if(!SelectedModuleInfo.IsValid())
+	{
+		const FString ProjectName = FApp::GetProjectName();
+
+		// Find initially selected module based on simple fallback in this order..
+		// Previously selected module, main project module, a  runtime module
+		TSharedPtr<FModuleContextInfo> ProjectModule;
+		TSharedPtr<FModuleContextInfo> RuntimeModule;
+
+		for (const auto& AvailableModule : AvailableModules)
+		{
+			// Check if this module matches our last used
+			if (AvailableModule->ModuleName == LastSelectedModuleName)
+			{
+				SelectedModuleInfo = AvailableModule;
+				break;
+			}
+
+			if (AvailableModule->ModuleName == ProjectName)
+			{
+				ProjectModule = AvailableModule;
+			}
+
+			if (AvailableModule->ModuleType == EHostType::Runtime)
+			{
+				RuntimeModule = AvailableModule;
+			}
+		}
+
+		if (!SelectedModuleInfo.IsValid())
+		{
+			if (ProjectModule.IsValid())
+			{
+				// use the project module we found
+				SelectedModuleInfo = ProjectModule;
+			}
+			else if (RuntimeModule.IsValid())
+			{
+				// use the first runtime module we found
+				SelectedModuleInfo = RuntimeModule;
+			}
+			else
+			{
+				// default to just the first module
+				SelectedModuleInfo = AvailableModules[0];
+			}
+		}
+
+		ModuleName = SelectedModuleInfo->ModuleSourcePath;
+	}
+}
+
+FText SWebAPIModulePicker::GetSelectedModuleComboText() const
+{
+	FFormatNamedArguments Args;
+	Args.Add(TEXT("ModuleName"), FText::FromString(SelectedModuleInfo->ModuleName));
+	Args.Add(TEXT("ModuleType"), FText::FromString(EHostType::ToString(SelectedModuleInfo->ModuleType)));
+	return FText::Format(LOCTEXT("ModuleComboEntry", "{ModuleName} ({ModuleType})"), Args);
+}
+
+void SWebAPIModulePicker::SelectedModuleComboBoxSelectionChanged(TSharedPtr<FModuleContextInfo> InValue, ESelectInfo::Type InSelectInfo)
+{
+	SelectedModuleInfo = InValue;
+	OnModuleChanged.ExecuteIfBound(SelectedModuleInfo->ModuleName, SelectedModuleInfo->ModuleSourcePath);
+}
+
+TSharedRef<SWidget> SWebAPIModulePicker::MakeWidgetForSelectedModuleCombo(TSharedPtr<FModuleContextInfo> InValue) const
+{
+	FFormatNamedArguments Args;
+	Args.Add(TEXT("ModuleName"), FText::FromString(InValue->ModuleName));
+	Args.Add(TEXT("ModuleType"), FText::FromString(EHostType::ToString(InValue->ModuleType)));
+	return SNew(STextBlock)
+		.Text(FText::Format(LOCTEXT("ModuleComboEntry", "{ModuleName} ({ModuleType})"), Args));
+}
+
+void SWebAPIModulePicker::OnModuleCreated(FString InModuleName)
+{
+	Refresh();
+}
+
+#undef LOCTEXT_NAMESPACE
diff --git a/Engine/Plugins/Web/WebAPI/Source/WebAPIEditor/Private/Details/Widgets/SWebAPIModulePicker.h b/Engine/Plugins/Web/WebAPI/Source/WebAPIEditor/Private/Details/Widgets/SWebAPIModulePicker.h
new file mode 100644
--- /dev/null
+++ b/Engine/Plugins/Web/WebAPI/Source/WebAPIEditor/Private/Details/Widgets/SWebAPIModulePicker.h
@@ -0,0 +1,69 @@
+﻿// Copyright Epic Games, Inc. All Rights Reserved.
+
+#pragma once
+
+#include "CoreMinimal.h"
+#include "ModuleDescriptor.h"
+#include "Widgets/DeclarativeSyntaxSupport.h"
+#include "Widgets/SCompoundWidget.h"
+#include "Widgets/SWidget.h"
+#include "Widgets/Input/SComboBox.h"
+
+class IClassViewerFilter;
+class SClassViewer;
+class SEditableTextBox;
+class SWizard;
+struct FParentClassItem;
+
+DECLARE_DELEGATE_TwoParams(FOnModuleChanged, const FString& /** Name */, const FString& /** Path */);
+
+/** A widget to select a target module or plugin for generated files. */
+class SWebAPIModulePicker : public SCompoundWidget
+{
+public:
+	SLATE_BEGIN_ARGS(SWebAPIModulePicker)
+	{}
+		/** The module initially provided. */
+		SLATE_ARGUMENT(FString, ModuleName)
+
+		/** Called whenever the text is committed.  This happens when the user presses enter or the text box loses focus. */
+		SLATE_EVENT(FOnModuleChanged, OnModuleChanged)
+	SLATE_END_ARGS()
+
+	/** Constructs this widget with InArgs. */
+	void Construct(const FArguments& InArgs);
+
+	/** Refreshes the list of available modules, accounting for any added or removed. */
+	void Refresh();
+
+private:
+	/** Get the combo box text for the currently selected module */
+	FText GetSelectedModuleComboText() const;
+
+	/** Called when the currently selected module is changed */
+	void SelectedModuleComboBoxSelectionChanged(TSharedPtr<FModuleContextInfo> InValue, ESelectInfo::Type InSelectInfo);
+
+	/** Create the widget to use as the combo box entry for the given module info */
+	TSharedRef<SWidget> MakeWidgetForSelectedModuleCombo(TSharedPtr<FModuleContextInfo> InValue) const;
+
+	/** When a module/plugin is newly created. */
+	void OnModuleCreated(FString InModuleName);
+	
+private:
+	/** The available modules combo box */
+	TSharedPtr<SComboBox<TSharedPtr<FModuleContextInfo>>> AvailableModulesCombo;
+
+	/** The last selected module name. Meant to keep the same module selected after first selection */
+	static FString LastSelectedModuleName;
+
+	/** The ModuleName chosen. */
+	FString ModuleName;
+
+	/** Information about the currently available modules for this project */
+	TArray<TSharedPtr<FModuleContextInfo>> AvailableModules;
+
+	/** Information about the currently selected module; used for class validation */
+	TSharedPtr<FModuleContextInfo> SelectedModuleInfo;
+
+    FOnModuleChanged OnModuleChanged;
+};
diff --git a/Engine/Plugins/Web/WebAPI/Source/WebAPIEditor/Private/Details/Widgets/SWebAPISchemaEnumRow.cpp b/Engine/Plugins/Web/WebAPI/Source/WebAPIEditor/Private/Details/Widgets/SWebAPISchemaEnumRow.cpp
new file mode 100644
--- /dev/null
+++ b/Engine/Plugins/Web/WebAPI/Source/WebAPIEditor/Private/Details/Widgets/SWebAPISchemaEnumRow.cpp
@@ -0,0 +1,81 @@
+﻿// Copyright Epic Games, Inc. All Rights Reserved.
+
+#include "SWebAPISchemaEnumRow.h"
+
+#include "WebAPIEditorStyle.h"
+#include "Details/ViewModels/WebAPIEnumViewModel.h"
+#include "Widgets/Colors/SColorBlock.h"
+#include "Widgets/Input/SHyperlink.h"
+#include "Widgets/Layout/SBox.h"
+
+#define LOCTEXT_NAMESPACE "WebAPISchemaEnumRow"
+
+void SWebAPISchemaEnumRow::Construct(const FArguments& InArgs, const TSharedRef<FWebAPIEnumViewModel>& InViewModel, const TSharedRef<STableViewBase>& InOwnerTableView)
+{
+	SWebAPISchemaTreeTableRow::Construct(
+		SWebAPISchemaTreeTableRow::FArguments()
+		.Content()
+		[
+			SNew(SBorder)
+			.ToolTipText(InViewModel->GetTooltip())
+			.BorderBackgroundColor(FLinearColor(0,0,0,0))
+			.Padding(4)
+			[
+				SNew(SBox)
+				[
+					SNew(SHorizontalBox)
+
+					// Model Type
+					+ SHorizontalBox::Slot()
+					.Padding(0)
+					.AutoWidth()
+					[
+						SNew(SBorder)
+						.IsEnabled(InViewModel, &FWebAPIEnumViewModel::GetShouldGenerate)
+						.BorderImage(FWebAPIEditorStyle::Get().GetBrush("WebAPI.TreeView.LabelBackground"))
+						[
+							SNew(SBox)
+							.WidthOverride(60)
+							.HeightOverride(20)
+							.HAlign(HAlign_Center)
+							.VAlign(VAlign_Center)
+							[
+								SNew(STextBlock)
+								.Text(LOCTEXT("ModelType", "Enum"))
+								.Font(FAppStyle::Get().GetFontStyle("SmallFontBold"))
+								.TransformPolicy(ETextTransformPolicy::ToUpper)
+							]
+						]
+					]		
+								
+					+ SHorizontalBox::Slot()
+					.Padding(0)
+					.FillWidth(1.0f)
+					[
+						SNew(SBox)
+						[
+							SNew(SHorizontalBox)
+							+ SHorizontalBox::Slot()
+							.Padding(9, 0, 0, 1)
+							.VAlign(VAlign_Center)
+							[
+								SNew(STextBlock)
+								.Text(InViewModel->GetLabel())
+								.Font(FAppStyle::Get().GetFontStyle("NormalFontBold"))
+							]
+							+ SHorizontalBox::Slot()
+							.VAlign(VAlign_Center)
+							.AutoWidth()
+							[
+								SNullWidget::NullWidget
+							]
+						]
+					]
+				]
+			]
+		],
+		InViewModel,
+		InOwnerTableView);
+}
+
+#undef LOCTEXT_NAMESPACE
diff --git a/Engine/Plugins/Web/WebAPI/Source/WebAPIEditor/Private/Details/Widgets/SWebAPISchemaEnumRow.h b/Engine/Plugins/Web/WebAPI/Source/WebAPIEditor/Private/Details/Widgets/SWebAPISchemaEnumRow.h
new file mode 100644
--- /dev/null
+++ b/Engine/Plugins/Web/WebAPI/Source/WebAPIEditor/Private/Details/Widgets/SWebAPISchemaEnumRow.h
@@ -0,0 +1,20 @@
+﻿// Copyright Epic Games, Inc. All Rights Reserved.
+
+#pragma once
+
+#include "CoreMinimal.h"
+#include "SWebAPISchemaTreeTableRow.h"
+
+class FWebAPIEnumViewModel;
+
+class SWebAPISchemaEnumRow
+	: public SWebAPISchemaTreeTableRow<FWebAPIEnumViewModel>
+{
+public:
+	SLATE_BEGIN_ARGS(SWebAPISchemaEnumRow)
+	{ }
+	SLATE_END_ARGS()
+
+public:
+	void Construct(const FArguments& InArgs, const TSharedRef<FWebAPIEnumViewModel>& InViewModel, const TSharedRef<STableViewBase>& InOwnerTableView);
+};
diff --git a/Engine/Plugins/Web/WebAPI/Source/WebAPIEditor/Private/Details/Widgets/SWebAPISchemaEnumValueRow.cpp b/Engine/Plugins/Web/WebAPI/Source/WebAPIEditor/Private/Details/Widgets/SWebAPISchemaEnumValueRow.cpp
new file mode 100644
--- /dev/null
+++ b/Engine/Plugins/Web/WebAPI/Source/WebAPIEditor/Private/Details/Widgets/SWebAPISchemaEnumValueRow.cpp
@@ -0,0 +1,47 @@
+﻿// Copyright Epic Games, Inc. All Rights Reserved.
+
+#include "SWebAPISchemaEnumValueRow.h"
+
+#include "Details/ViewModels/WebAPIEnumViewModel.h"
+#include "Widgets/Layout/SBox.h"
+
+void SWebAPISchemaEnumValueRow::Construct(const FArguments& InArgs, const TSharedRef<FWebAPIEnumValueViewModel>& InViewModel, const TSharedRef<STableViewBase>& InOwnerTableView)
+{
+	FSlateColor TypeColor = FSlateColor();
+	
+	SWebAPISchemaTreeTableRow::Construct(
+	SWebAPISchemaTreeTableRow::FArguments()
+	.Content()
+	[
+		SNew(SBorder)
+		.ToolTipText(InViewModel->GetTooltip())
+		.BorderBackgroundColor(FLinearColor(0,0,0,0))
+		.Padding(4)
+		[
+			SNew(SBox)
+			[
+				SNew(SHorizontalBox)
+
+				+ SHorizontalBox::Slot()
+				.VAlign(VAlign_Fill)
+				.Padding(0)
+				[
+					SNew(SBox)
+					[
+						SNew(SHorizontalBox)						
+						+ SHorizontalBox::Slot()
+						.Padding(9, 0, 0, 1)
+						.VAlign(VAlign_Center)
+						[
+							SNew(STextBlock)
+							.TextStyle(FEditorStyle::Get(), "PlacementBrowser.Asset.Name")
+							.Text(InViewModel->GetLabel())
+						]
+					]
+				]
+			]
+		]
+	],
+	InViewModel,
+	InOwnerTableView);
+}
diff --git a/Engine/Plugins/Web/WebAPI/Source/WebAPIEditor/Private/Details/Widgets/SWebAPISchemaEnumValueRow.h b/Engine/Plugins/Web/WebAPI/Source/WebAPIEditor/Private/Details/Widgets/SWebAPISchemaEnumValueRow.h
new file mode 100644
--- /dev/null
+++ b/Engine/Plugins/Web/WebAPI/Source/WebAPIEditor/Private/Details/Widgets/SWebAPISchemaEnumValueRow.h
@@ -0,0 +1,19 @@
+﻿// Copyright Epic Games, Inc. All Rights Reserved.
+
+#pragma once
+
+#include "CoreMinimal.h"
+#include "SWebAPISchemaTreeTableRow.h"
+#include "Details/ViewModels/WebAPIEnumViewModel.h"
+
+class SWebAPISchemaEnumValueRow
+	: public SWebAPISchemaTreeTableRow<FWebAPIEnumValueViewModel>
+{
+public:
+	SLATE_BEGIN_ARGS(SWebAPISchemaEnumValueRow)
+	{ }
+	SLATE_END_ARGS()
+
+public:
+	void Construct(const FArguments& InArgs, const TSharedRef<FWebAPIEnumValueViewModel>& InViewModel, const TSharedRef<STableViewBase>& InOwnerTableView);
+};
diff --git a/Engine/Plugins/Web/WebAPI/Source/WebAPIEditor/Private/Details/Widgets/SWebAPISchemaModelRow.cpp b/Engine/Plugins/Web/WebAPI/Source/WebAPIEditor/Private/Details/Widgets/SWebAPISchemaModelRow.cpp
new file mode 100644
--- /dev/null
+++ b/Engine/Plugins/Web/WebAPI/Source/WebAPIEditor/Private/Details/Widgets/SWebAPISchemaModelRow.cpp
@@ -0,0 +1,77 @@
+﻿// Copyright Epic Games, Inc. All Rights Reserved.
+
+#include "SWebAPISchemaModelRow.h"
+
+#include "WebAPIEditorStyle.h"
+#include "Details/ViewModels/WebAPIModelViewModel.h"
+#include "Widgets/Colors/SColorBlock.h"
+#include "Widgets/Input/SHyperlink.h"
+#include "Widgets/Layout/SBox.h"
+
+#define LOCTEXT_NAMESPACE "WebAPISchemaModelRow"
+
+void SWebAPISchemaModelRow::Construct(const FArguments& InArgs, const TSharedRef<FWebAPIModelViewModel>& InViewModel, const TSharedRef<STableViewBase>& InOwnerTableView)
+{
+	SWebAPISchemaTreeTableRow::Construct(
+		SWebAPISchemaTreeTableRow::FArguments()
+		.IsEnabled(InViewModel, &FWebAPIModelViewModel::GetShouldGenerate)
+		.Content()
+		[
+			SNew(SBox)
+			[
+				SNew(SHorizontalBox)
+
+				// Model Type
+				+ SHorizontalBox::Slot()
+				.Padding(0)
+				.AutoWidth()
+				[
+					SNew(SBorder)
+					.IsEnabled(InViewModel, &FWebAPIModelViewModel::GetShouldGenerate)
+					.BorderImage(FWebAPIEditorStyle::Get().GetBrush("WebAPI.TreeView.LabelBackground"))
+					[
+						SNew(SBox)
+						.WidthOverride(60)
+						.HeightOverride(20)
+						.HAlign(HAlign_Center)
+						.VAlign(VAlign_Center)
+						[
+							SNew(STextBlock)
+							.Text(LOCTEXT("ModelType", "Model"))
+							.Font(FAppStyle::Get().GetFontStyle("SmallFontBold"))
+							.TransformPolicy(ETextTransformPolicy::ToUpper)
+						]
+					]
+				]
+		
+				+ SHorizontalBox::Slot()
+				.Padding(0)
+				.FillWidth(1.0f)
+				[
+					SNew(SBox)
+					.IsEnabled(InViewModel, &FWebAPIModelViewModel::GetShouldGenerate)
+					[
+						SNew(SHorizontalBox)
+						+ SHorizontalBox::Slot()
+						.Padding(9, 0, 0, 1)
+						.VAlign(VAlign_Center)
+						[
+							SNew(STextBlock)
+							.Text(InViewModel->GetLabel())
+							.Font(FAppStyle::Get().GetFontStyle("NormalFontBold"))
+						]
+						+ SHorizontalBox::Slot()
+						.VAlign(VAlign_Center)
+						.AutoWidth()
+						[
+							SNullWidget::NullWidget
+						]
+					]
+				]
+			]
+		],
+		InViewModel,
+		InOwnerTableView);
+}
+
+#undef LOCTEXT_NAMESPACE
diff --git a/Engine/Plugins/Web/WebAPI/Source/WebAPIEditor/Private/Details/Widgets/SWebAPISchemaModelRow.h b/Engine/Plugins/Web/WebAPI/Source/WebAPIEditor/Private/Details/Widgets/SWebAPISchemaModelRow.h
new file mode 100644
--- /dev/null
+++ b/Engine/Plugins/Web/WebAPI/Source/WebAPIEditor/Private/Details/Widgets/SWebAPISchemaModelRow.h
@@ -0,0 +1,20 @@
+﻿// Copyright Epic Games, Inc. All Rights Reserved.
+
+#pragma once
+
+#include "CoreMinimal.h"
+#include "SWebAPISchemaTreeTableRow.h"
+
+class FWebAPIModelViewModel;
+
+class SWebAPISchemaModelRow
+	: public SWebAPISchemaTreeTableRow<FWebAPIModelViewModel>
+{
+public:
+	SLATE_BEGIN_ARGS(SWebAPISchemaModelRow)
+	{ }
+	SLATE_END_ARGS()
+
+public:
+	void Construct(const FArguments& InArgs, const TSharedRef<FWebAPIModelViewModel>& InViewModel, const TSharedRef<STableViewBase>& InOwnerTableView);
+};
diff --git a/Engine/Plugins/Web/WebAPI/Source/WebAPIEditor/Private/Details/Widgets/SWebAPISchemaOperationParameterRow.cpp b/Engine/Plugins/Web/WebAPI/Source/WebAPIEditor/Private/Details/Widgets/SWebAPISchemaOperationParameterRow.cpp
new file mode 100644
--- /dev/null
+++ b/Engine/Plugins/Web/WebAPI/Source/WebAPIEditor/Private/Details/Widgets/SWebAPISchemaOperationParameterRow.cpp
@@ -0,0 +1,71 @@
+﻿ // Copyright Epic Games, Inc. All Rights Reserved.
+
+#include "SWebAPISchemaOperationParameterRow.h"
+
+#include "Details/ViewModels/WebAPIOperationParameterViewModel.h"
+#include "Widgets/Layout/SBox.h"
+
+void SWebAPISchemaOperationParameterRow::Construct(const FArguments& InArgs, const TSharedRef<FWebAPIOperationParameterViewModel>& InViewModel, const TSharedRef<STableViewBase>& InOwnerTableView)
+{
+	const FSlateBrush* ArrayIcon = FEditorStyle::GetBrush(TEXT("Graph.ArrayPin.Connected"));
+	const FSlateBrush* Icon = FEditorStyle::GetBrush(TEXT("Graph.Pin.Connected"));
+	
+	const FSlateColor TypeColor = InViewModel->GetPinColor();
+	
+	SWebAPISchemaTreeTableRow::Construct(
+	SWebAPISchemaTreeTableRow::FArguments()
+	.Content()
+	[
+		SNew(SBorder)
+		.ToolTipText(InViewModel->GetTooltip())
+		.BorderBackgroundColor(FLinearColor(0,0,0,0))
+		.Padding(4)
+		[
+			SNew(SBox)
+			[
+				SNew(SHorizontalBox)
+				
+				+ SHorizontalBox::Slot()
+				.Padding(0)
+				.AutoWidth()
+				[
+					SNew(SBox)
+					.WidthOverride(20)
+					.HeightOverride(20)
+					.HAlign(HAlign_Center)
+					.VAlign(VAlign_Center)
+					[
+						SNew(SImage)
+						.Image_Lambda([&, ArrayIcon, Icon]
+						{
+							return ViewModel->IsArray()
+									? ArrayIcon
+									: Icon;
+						})
+						.ColorAndOpacity(TypeColor)
+					]
+				]
+				
+				+ SHorizontalBox::Slot()
+				.VAlign(VAlign_Fill)
+				.Padding(0)
+				[
+					SNew(SBox)
+					[
+						SNew(SHorizontalBox)
+						+ SHorizontalBox::Slot()
+						.Padding(9, 0, 0, 1)
+						.VAlign(VAlign_Center)
+						[
+							SNew(STextBlock)
+							.TextStyle(FEditorStyle::Get(), "PlacementBrowser.Asset.Name")
+							.Text(InViewModel->GetLabel())
+						]
+					]
+				]
+			]
+		]
+	],
+	InViewModel,
+	InOwnerTableView);
+}
diff --git a/Engine/Plugins/Web/WebAPI/Source/WebAPIEditor/Private/Details/Widgets/SWebAPISchemaOperationParameterRow.h b/Engine/Plugins/Web/WebAPI/Source/WebAPIEditor/Private/Details/Widgets/SWebAPISchemaOperationParameterRow.h
new file mode 100644
--- /dev/null
+++ b/Engine/Plugins/Web/WebAPI/Source/WebAPIEditor/Private/Details/Widgets/SWebAPISchemaOperationParameterRow.h
@@ -0,0 +1,19 @@
+﻿// Copyright Epic Games, Inc. All Rights Reserved.
+
+#pragma once
+
+#include "CoreMinimal.h"
+#include "SWebAPISchemaTreeTableRow.h"
+#include "Details/ViewModels/WebAPIOperationParameterViewModel.h"
+
+class SWebAPISchemaOperationParameterRow
+	: public SWebAPISchemaTreeTableRow<FWebAPIOperationParameterViewModel>
+{
+public:
+	SLATE_BEGIN_ARGS(SWebAPISchemaOperationParameterRow)
+	{ }
+	SLATE_END_ARGS()
+
+public:
+	void Construct(const FArguments& InArgs, const TSharedRef<FWebAPIOperationParameterViewModel>& InViewModel, const TSharedRef<STableViewBase>& InOwnerTableView);
+};
diff --git a/Engine/Plugins/Web/WebAPI/Source/WebAPIEditor/Private/Details/Widgets/SWebAPISchemaOperationRequestRow.cpp b/Engine/Plugins/Web/WebAPI/Source/WebAPIEditor/Private/Details/Widgets/SWebAPISchemaOperationRequestRow.cpp
new file mode 100644
--- /dev/null
+++ b/Engine/Plugins/Web/WebAPI/Source/WebAPIEditor/Private/Details/Widgets/SWebAPISchemaOperationRequestRow.cpp
@@ -0,0 +1,60 @@
+﻿// Copyright Epic Games, Inc. All Rights Reserved.
+
+#include "SWebAPISchemaOperationRequestRow.h"
+
+#include "WebAPIEditorStyle.h"
+#include "Details/ViewModels/WebAPIEnumViewModel.h"
+#include "Widgets/Colors/SColorBlock.h"
+#include "Widgets/Input/SHyperlink.h"
+#include "Widgets/Layout/SBox.h"
+
+void SWebAPISchemaOperationRequestRow::Construct(const FArguments& InArgs, const TSharedRef<FWebAPIOperationRequestViewModel>& InViewModel, const TSharedRef<STableViewBase>& InOwnerTableView)
+{
+	SWebAPISchemaTreeTableRow::Construct(
+		SWebAPISchemaTreeTableRow::FArguments()
+		.Content()
+		[
+			SNew(SBorder)
+			.ToolTipText(InViewModel->GetTooltip())
+			.BorderBackgroundColor(FLinearColor(0,0,0,0))
+			.Padding(4)
+			[
+				SNew(SBox)
+				[
+					SNew(SHorizontalBox)
+								
+					+ SHorizontalBox::Slot()
+					.Padding(0)
+					.FillWidth(1.0f)
+					[
+						SNew(SBorder)
+						.BorderBackgroundColor(FLinearColor(0,0,0,0))
+						[
+							SNew(SBox)
+							[
+								SNew(SHorizontalBox)
+								
+								+ SHorizontalBox::Slot()
+								.Padding(9, 0, 0, 1)
+								.VAlign(VAlign_Center)
+								[
+									SNew(STextBlock)
+									.Text(InViewModel->GetLabel())
+									.Font(FAppStyle::Get().GetFontStyle("NormalFontBold"))
+								]
+								
+								+ SHorizontalBox::Slot()
+								.VAlign(VAlign_Center)
+								.AutoWidth()
+								[
+									SNullWidget::NullWidget
+								]
+							]
+						]
+					]
+				]
+			]
+		],
+		InViewModel,
+		InOwnerTableView);
+}
diff --git a/Engine/Plugins/Web/WebAPI/Source/WebAPIEditor/Private/Details/Widgets/SWebAPISchemaOperationRequestRow.h b/Engine/Plugins/Web/WebAPI/Source/WebAPIEditor/Private/Details/Widgets/SWebAPISchemaOperationRequestRow.h
new file mode 100644
--- /dev/null
+++ b/Engine/Plugins/Web/WebAPI/Source/WebAPIEditor/Private/Details/Widgets/SWebAPISchemaOperationRequestRow.h
@@ -0,0 +1,21 @@
+﻿// Copyright Epic Games, Inc. All Rights Reserved.
+
+#pragma once
+
+#include "CoreMinimal.h"
+#include "SWebAPISchemaTreeTableRow.h"
+#include "Details/ViewModels/WebAPIOperationRequestViewModel.h"
+
+class FWebAPIEnumViewModel;
+
+class SWebAPISchemaOperationRequestRow
+	: public SWebAPISchemaTreeTableRow<FWebAPIOperationRequestViewModel>
+{
+public:
+	SLATE_BEGIN_ARGS(SWebAPISchemaOperationRequestRow)
+	{ }
+	SLATE_END_ARGS()
+
+public:
+	void Construct(const FArguments& InArgs, const TSharedRef<FWebAPIOperationRequestViewModel>& InViewModel, const TSharedRef<STableViewBase>& InOwnerTableView);
+};
diff --git a/Engine/Plugins/Web/WebAPI/Source/WebAPIEditor/Private/Details/Widgets/SWebAPISchemaOperationResponseRow.cpp b/Engine/Plugins/Web/WebAPI/Source/WebAPIEditor/Private/Details/Widgets/SWebAPISchemaOperationResponseRow.cpp
new file mode 100644
--- /dev/null
+++ b/Engine/Plugins/Web/WebAPI/Source/WebAPIEditor/Private/Details/Widgets/SWebAPISchemaOperationResponseRow.cpp
@@ -0,0 +1,61 @@
+﻿// Copyright Epic Games, Inc. All Rights Reserved.
+
+#include "SWebAPISchemaOperationResponseRow.h"
+
+#include "WebAPIEditorStyle.h"
+#include "Details/ViewModels/WebAPIEnumViewModel.h"
+#include "Widgets/Colors/SColorBlock.h"
+#include "Widgets/Input/SHyperlink.h"
+#include "Widgets/Layout/SBox.h"
+
+void SWebAPISchemaOperationResponseRow::Construct(const FArguments& InArgs, const TSharedRef<FWebAPIOperationResponseViewModel>& InViewModel, const TSharedRef<STableViewBase>& InOwnerTableView)
+{
+	SWebAPISchemaTreeTableRow::Construct(
+		SWebAPISchemaTreeTableRow::FArguments()
+		.Content()
+		[
+			SNew(SBorder)
+			.ToolTipText(InViewModel->GetTooltip())
+			.BorderBackgroundColor(FLinearColor(0,0,0,0))
+			.Padding(4)
+			[
+				SNew(SBox)
+				[
+					SNew(SHorizontalBox)
+								
+					+ SHorizontalBox::Slot()
+					.Padding(0)
+					.FillWidth(1.0f)
+					[
+						SNew(SBorder)
+						.BorderBackgroundColor(FLinearColor(0,0,0,0))
+						[
+							SNew(SBox)
+							[
+								SNew(SHorizontalBox)
+								
+								+ SHorizontalBox::Slot()
+								.Padding(9, 0, 0, 1)
+								.AutoWidth()
+								.VAlign(VAlign_Center)
+								[
+									SNew(STextBlock)
+									.Text(InViewModel->GetLabel())
+									.Font(FAppStyle::Get().GetFontStyle("NormalFontBold"))
+								]
+
+								+ SHorizontalBox::Slot()
+								.VAlign(VAlign_Center)
+								.AutoWidth()
+								[
+									SNullWidget::NullWidget
+								]
+							]
+						]
+					]
+				]
+			]
+		],
+		InViewModel,
+		InOwnerTableView);
+}
diff --git a/Engine/Plugins/Web/WebAPI/Source/WebAPIEditor/Private/Details/Widgets/SWebAPISchemaOperationResponseRow.h b/Engine/Plugins/Web/WebAPI/Source/WebAPIEditor/Private/Details/Widgets/SWebAPISchemaOperationResponseRow.h
new file mode 100644
--- /dev/null
+++ b/Engine/Plugins/Web/WebAPI/Source/WebAPIEditor/Private/Details/Widgets/SWebAPISchemaOperationResponseRow.h
@@ -0,0 +1,21 @@
+﻿// Copyright Epic Games, Inc. All Rights Reserved.
+
+#pragma once
+
+#include "CoreMinimal.h"
+#include "SWebAPISchemaTreeTableRow.h"
+#include "Details/ViewModels/WebAPIOperationResponseViewModel.h"
+
+class FWebAPIEnumViewModel;
+
+class SWebAPISchemaOperationResponseRow
+	: public SWebAPISchemaTreeTableRow<FWebAPIOperationResponseViewModel>
+{
+public:
+	SLATE_BEGIN_ARGS(SWebAPISchemaOperationResponseRow)
+	{ }
+	SLATE_END_ARGS()
+
+public:
+	void Construct(const FArguments& InArgs, const TSharedRef<FWebAPIOperationResponseViewModel>& InViewModel, const TSharedRef<STableViewBase>& InOwnerTableView);
+};
diff --git a/Engine/Plugins/Web/WebAPI/Source/WebAPIEditor/Private/Details/Widgets/SWebAPISchemaOperationRow.cpp b/Engine/Plugins/Web/WebAPI/Source/WebAPIEditor/Private/Details/Widgets/SWebAPISchemaOperationRow.cpp
new file mode 100644
--- /dev/null
+++ b/Engine/Plugins/Web/WebAPI/Source/WebAPIEditor/Private/Details/Widgets/SWebAPISchemaOperationRow.cpp
@@ -0,0 +1,127 @@
+﻿// Copyright Epic Games, Inc. All Rights Reserved.
+
+#include "SWebAPISchemaOperationRow.h"
+
+#include "WebAPIEditorStyle.h"
+#include "Details/ViewModels/WebAPIOperationViewModel.h"
+#include "Widgets/Colors/SColorBlock.h"
+#include "Widgets/Input/SCheckBox.h"
+#include "Widgets/Input/SHyperlink.h"
+#include "Widgets/Layout/SBox.h"
+#include "Widgets/Text/SRichTextBlock.h"
+
+void SWebAPISchemaOperationRow::Construct(const FArguments& InArgs, const TSharedRef<FWebAPIOperationViewModel>& InViewModel, const TSharedRef<STableViewBase>& InOwnerTableView)
+{
+	SWebAPISchemaTreeTableRow::Construct(
+		SWebAPISchemaTreeTableRow::FArguments()
+		.IsEnabled(InViewModel, &FWebAPIOperationViewModel::GetShouldGenerate, true)
+		.Content()
+		[
+			SNew(SBorder)
+			.BorderBackgroundColor(FLinearColor(0,0,0,0))
+			// @todo: this prevents selection highlight
+			//.BorderImage(FEditorStyle::GetBrush("ToolPanel.GroupBorder"))
+			.Padding(4)
+			[
+				SNew(SBox)
+				[
+					SNew(SVerticalBox)
+
+					// Main row
+					+ SVerticalBox::Slot()
+					.Padding(0)
+					[
+						SNew(SHorizontalBox)
+
+						// Generation Toggle
+						+ SHorizontalBox::Slot()
+						.Padding(0)
+						.AutoWidth()
+						[
+							SNew(SBox)
+							.IsEnabled(InViewModel->GetService().Get(), &FWebAPIServiceViewModel::GetShouldGenerate)
+							.WidthOverride(20)
+							.HeightOverride(20)
+							.HAlign(HAlign_Center)
+							.VAlign(VAlign_Center)
+							[
+								SNew(SCheckBox)
+								.OnCheckStateChanged_Lambda([InViewModel](const ECheckBoxState InCheckState) { InViewModel->SetShouldGenerate(InCheckState == ECheckBoxState::Checked); })
+								.IsChecked_Lambda([InViewModel]
+									{ return InViewModel->GetShouldGenerate(false) ? ECheckBoxState::Checked : ECheckBoxState::Unchecked; })
+							]
+						]
+
+						// HTTP Verb
+						+ SHorizontalBox::Slot()
+						.Padding(0)
+						.AutoWidth()
+						[
+							SNew(SBorder)
+							.IsEnabled(InViewModel, &FWebAPIOperationViewModel::GetShouldGenerate, true)
+							.BorderImage(FWebAPIEditorStyle::Get().GetBrush("WebAPI.TreeView.LabelBackground"))
+							[
+								SNew(SBox)
+								.WidthOverride(60)
+								.HeightOverride(20)
+								.HAlign(HAlign_Center)
+								.VAlign(VAlign_Center)
+								[
+									SNew(STextBlock)
+									.Text(InViewModel->GetVerb())
+									.Font(FAppStyle::Get().GetFontStyle("SmallFontBold"))
+									.TransformPolicy(ETextTransformPolicy::ToUpper)
+								]
+							]
+						]
+
+						// Operation Label
+						+ SHorizontalBox::Slot()
+						.VAlign(VAlign_Fill)
+						.Padding(0)
+						[
+							SNew(SBox)
+							.IsEnabled(InViewModel, &FWebAPIOperationViewModel::GetShouldGenerate, true)
+							[
+								SNew(SHorizontalBox)
+								+ SHorizontalBox::Slot()
+								.Padding(9, 0, 0, 1)
+								.VAlign(VAlign_Center)
+								[
+									SNew(STextBlock)
+									.Text(InViewModel->GetLabel())
+									.Font(FAppStyle::Get().GetFontStyle("NormalFontBold"))
+									.TextStyle(FEditorStyle::Get(), "PlacementBrowser.Asset.Name")								
+								]
+								+ SHorizontalBox::Slot()
+								.VAlign(VAlign_Center)
+								.AutoWidth()
+								[
+									SNullWidget::NullWidget
+								] 
+							]
+						]
+					]
+
+					// Path row
+					+ SVerticalBox::Slot()
+					.Padding(0)
+					[
+						SNew(SHorizontalBox)
+
+						// Operation Path
+						+ SHorizontalBox::Slot()
+						// 20 = Generation Toggle checkbox width
+						.Padding(20, 9, 0, 1)
+						.VAlign(VAlign_Center) 
+						[
+							SNew(SRichTextBlock)
+							.Text(InViewModel->GetRichPath())
+						]
+					]
+				]
+			]
+		],
+		InViewModel,
+		InOwnerTableView);
+}
diff --git a/Engine/Plugins/Web/WebAPI/Source/WebAPIEditor/Private/Details/Widgets/SWebAPISchemaOperationRow.h b/Engine/Plugins/Web/WebAPI/Source/WebAPIEditor/Private/Details/Widgets/SWebAPISchemaOperationRow.h
new file mode 100644
--- /dev/null
+++ b/Engine/Plugins/Web/WebAPI/Source/WebAPIEditor/Private/Details/Widgets/SWebAPISchemaOperationRow.h
@@ -0,0 +1,19 @@
+﻿// Copyright Epic Games, Inc. All Rights Reserved.
+
+#pragma once
+
+#include "CoreMinimal.h"
+#include "SWebAPISchemaTreeTableRow.h"
+#include "Details/ViewModels/WebAPIOperationViewModel.h"
+
+class SWebAPISchemaOperationRow
+	: public SWebAPISchemaTreeTableRow<FWebAPIOperationViewModel>
+{
+public:
+	SLATE_BEGIN_ARGS(SWebAPISchemaOperationRow)
+	{ }
+	SLATE_END_ARGS()
+
+public:
+	void Construct(const FArguments& InArgs, const TSharedRef<FWebAPIOperationViewModel>& InViewModel, const TSharedRef<STableViewBase>& InOwnerTableView);
+};
diff --git a/Engine/Plugins/Web/WebAPI/Source/WebAPIEditor/Private/Details/Widgets/SWebAPISchemaPropertyRow.cpp b/Engine/Plugins/Web/WebAPI/Source/WebAPIEditor/Private/Details/Widgets/SWebAPISchemaPropertyRow.cpp
new file mode 100644
--- /dev/null
+++ b/Engine/Plugins/Web/WebAPI/Source/WebAPIEditor/Private/Details/Widgets/SWebAPISchemaPropertyRow.cpp
@@ -0,0 +1,69 @@
+﻿// Copyright Epic Games, Inc. All Rights Reserved.
+
+#include "SWebAPISchemaPropertyRow.h"
+
+#include "Widgets/Layout/SBox.h"
+
+void SWebAPISchemaPropertyRow::Construct(const FArguments& InArgs, const TSharedRef<FWebAPIPropertyViewModel>& InViewModel, const TSharedRef<STableViewBase>& InOwnerTableView)
+{
+	const FSlateBrush* ArrayIcon = FEditorStyle::GetBrush(TEXT("Graph.ArrayPin.Connected"));
+	const FSlateBrush* Icon = FEditorStyle::GetBrush(TEXT("Graph.Pin.Connected"));
+	const FSlateColor TypeColor = InViewModel->GetPinColor();
+	
+	SWebAPISchemaTreeTableRow::Construct(
+	SWebAPISchemaTreeTableRow::FArguments()
+	.Content()
+	[
+		SNew(SBorder)
+		.ToolTipText(InViewModel->GetTooltip())
+		.BorderBackgroundColor(FLinearColor(0,0,0,0))		
+		.Padding(4)
+		[
+			SNew(SBox)
+			[
+				SNew(SHorizontalBox)
+				
+				+ SHorizontalBox::Slot()
+				.Padding(0)
+				.AutoWidth()
+				[
+					SNew(SBox)
+					.WidthOverride(20)
+					.HeightOverride(20)
+					.HAlign(HAlign_Center)
+					.VAlign(VAlign_Center)
+					[
+						SNew(SImage)
+						.Image_Lambda([&, ArrayIcon, Icon]
+						{
+							return ViewModel->IsArray()
+									? ArrayIcon
+									: Icon;
+						})
+						.ColorAndOpacity(TypeColor)
+					]
+				]
+				
+				+ SHorizontalBox::Slot()
+				.VAlign(VAlign_Fill)
+				.Padding(0)
+				[
+					SNew(SBox)
+					[
+						SNew(SHorizontalBox)
+						+ SHorizontalBox::Slot()
+						.Padding(9, 0, 0, 1)
+						.VAlign(VAlign_Center)
+						[
+							SNew(STextBlock)
+							.TextStyle(FEditorStyle::Get(), "PlacementBrowser.Asset.Name")
+							.Text(InViewModel->GetLabel())
+						]
+					]
+				]
+			]
+		]
+	],
+	InViewModel,
+	InOwnerTableView);
+}
diff --git a/Engine/Plugins/Web/WebAPI/Source/WebAPIEditor/Private/Details/Widgets/SWebAPISchemaPropertyRow.h b/Engine/Plugins/Web/WebAPI/Source/WebAPIEditor/Private/Details/Widgets/SWebAPISchemaPropertyRow.h
new file mode 100644
--- /dev/null
+++ b/Engine/Plugins/Web/WebAPI/Source/WebAPIEditor/Private/Details/Widgets/SWebAPISchemaPropertyRow.h
@@ -0,0 +1,19 @@
+﻿// Copyright Epic Games, Inc. All Rights Reserved.
+
+#pragma once
+
+#include "CoreMinimal.h"
+#include "SWebAPISchemaTreeTableRow.h"
+#include "Details/ViewModels/WebAPIModelViewModel.h"
+
+class SWebAPISchemaPropertyRow
+	: public SWebAPISchemaTreeTableRow<FWebAPIPropertyViewModel>
+{
+public:
+	SLATE_BEGIN_ARGS(SWebAPISchemaPropertyRow)
+	{ }
+	SLATE_END_ARGS()
+
+public:
+	void Construct(const FArguments& InArgs, const TSharedRef<FWebAPIPropertyViewModel>& InViewModel, const TSharedRef<STableViewBase>& InOwnerTableView);
+};
diff --git a/Engine/Plugins/Web/WebAPI/Source/WebAPIEditor/Private/Details/Widgets/SWebAPISchemaServiceRow.cpp b/Engine/Plugins/Web/WebAPI/Source/WebAPIEditor/Private/Details/Widgets/SWebAPISchemaServiceRow.cpp
new file mode 100644
--- /dev/null
+++ b/Engine/Plugins/Web/WebAPI/Source/WebAPIEditor/Private/Details/Widgets/SWebAPISchemaServiceRow.cpp
@@ -0,0 +1,73 @@
+﻿// Copyright Epic Games, Inc. All Rights Reserved.
+
+#include "SWebAPISchemaServiceRow.h"
+
+#include "WebAPIEditorStyle.h"
+#include "Widgets/Colors/SColorBlock.h"
+#include "Widgets/Input/SCheckBox.h"
+#include "Widgets/Input/SHyperlink.h"
+#include "Widgets/Layout/SBox.h"
+
+void SWebAPISchemaServiceRow::Construct(const FArguments& InArgs, const TSharedRef<FWebAPIServiceViewModel>& InViewModel, const TSharedRef<STableViewBase>& InOwnerTableView)
+{
+	SWebAPISchemaTreeTableRow::Construct(
+		SWebAPISchemaTreeTableRow::FArguments()
+		.Content()
+		[
+			SNew(SBorder)
+			.BorderBackgroundColor(FLinearColor(0,0,0,0))
+			// @todo: this prevents selection highlight
+			//.BorderImage(FEditorStyle::GetBrush("ToolPanel.GroupBorder"))
+			.Padding(4)
+			[
+				SNew(SBox)
+				[
+					SNew(SHorizontalBox)
+					+ SHorizontalBox::Slot()
+					.Padding(0)
+					.AutoWidth()
+					[
+						SNew(SBox)
+						.WidthOverride(20)
+						.HeightOverride(20)
+						.HAlign(HAlign_Center)
+						.VAlign(VAlign_Center)
+						[
+							SNew(SCheckBox)
+							.OnCheckStateChanged_Lambda([InViewModel](const ECheckBoxState InCheckState) { InViewModel->SetShouldGenerate(InCheckState == ECheckBoxState::Checked); })
+							.IsChecked_Lambda([InViewModel]
+								{ return InViewModel->GetShouldGenerate() ? ECheckBoxState::Checked : ECheckBoxState::Unchecked; })
+						]
+					]
+								
+					+ SHorizontalBox::Slot()
+					.Padding(0)
+					.FillWidth(1.0f)
+					[
+						SNew(SBox)
+						.IsEnabled(InViewModel, &FWebAPIServiceViewModel::GetShouldGenerate)
+						[
+							SNew(SHorizontalBox)
+							+ SHorizontalBox::Slot()
+							.Padding(9, 0, 0, 1)
+							.VAlign(VAlign_Center)
+							[
+								SNew(STextBlock)
+								.IsEnabled(InViewModel, &FWebAPIServiceViewModel::GetShouldGenerate)
+								.Text(InViewModel->GetLabel())
+								.Font(FAppStyle::Get().GetFontStyle("NormalFontBold"))
+							]
+							+ SHorizontalBox::Slot()
+							.VAlign(VAlign_Center)
+							.AutoWidth()
+							[
+								SNullWidget::NullWidget
+							]
+						]
+					]
+				]
+			]
+		],
+		InViewModel,
+		InOwnerTableView);
+}
diff --git a/Engine/Plugins/Web/WebAPI/Source/WebAPIEditor/Private/Details/Widgets/SWebAPISchemaServiceRow.h b/Engine/Plugins/Web/WebAPI/Source/WebAPIEditor/Private/Details/Widgets/SWebAPISchemaServiceRow.h
new file mode 100644
--- /dev/null
+++ b/Engine/Plugins/Web/WebAPI/Source/WebAPIEditor/Private/Details/Widgets/SWebAPISchemaServiceRow.h
@@ -0,0 +1,19 @@
+﻿// Copyright Epic Games, Inc. All Rights Reserved.
+
+#pragma once
+
+#include "CoreMinimal.h"
+#include "SWebAPISchemaTreeTableRow.h"
+#include "Details/ViewModels/WebAPIServiceViewModel.h"
+
+class SWebAPISchemaServiceRow
+	: public SWebAPISchemaTreeTableRow<FWebAPIServiceViewModel>
+{
+public:
+	SLATE_BEGIN_ARGS(SWebAPISchemaServiceRow)
+	{ }
+	SLATE_END_ARGS()
+
+public:
+	void Construct(const FArguments& InArgs, const TSharedRef<FWebAPIServiceViewModel>& InViewModel, const TSharedRef<STableViewBase>& InOwnerTableView);
+};
diff --git a/Engine/Plugins/Web/WebAPI/Source/WebAPIEditor/Private/Details/Widgets/SWebAPISchemaTreeTableRow.cpp b/Engine/Plugins/Web/WebAPI/Source/WebAPIEditor/Private/Details/Widgets/SWebAPISchemaTreeTableRow.cpp
new file mode 100644
--- /dev/null
+++ b/Engine/Plugins/Web/WebAPI/Source/WebAPIEditor/Private/Details/Widgets/SWebAPISchemaTreeTableRow.cpp
@@ -0,0 +1,47 @@
+﻿ // Copyright Epic Games, Inc. All Rights Reserved.
+
+#include "SWebAPISchemaTreeTableRow.h"
+
+#include "Details/ViewModels/WebAPIViewModel.h"
+
+#define LOCTEXT_NAMESPACE "WebAPISchemaTreeItemWidget"
+
+constexpr float IndentSize = 12;
+
+TSharedPtr<IWebAPIViewModel> GetEntryAbove(TSharedPtr<IWebAPIViewModel> InEntry)
+{
+	const TSharedPtr<IWebAPIViewModel> Outer = InEntry->GetParent();
+	if (Outer)
+	{
+		TArray<TSharedPtr<IWebAPIViewModel>> Children;
+		Outer->GetChildren(Children);
+		for (int32 Index = 1; Index < Children.Num(); Index++)
+		{
+			if (Children[Index] == InEntry)
+			{
+				return Children[Index - 1];
+			}
+		}
+	}
+	return nullptr;
+}
+
+TSharedPtr<IWebAPIViewModel> GetEntryBelow(TSharedPtr<IWebAPIViewModel> InEntry)
+{
+	const TSharedPtr<IWebAPIViewModel> Outer = InEntry->GetParent();
+	if (Outer)
+	{
+		TArray<TSharedPtr<IWebAPIViewModel>> Children;
+		Outer->GetChildren(Children);
+		for (int32 Index = 0; Index < Children.Num() - 1; Index++)
+		{
+			if (Children[Index] == InEntry)
+			{
+				return Children[Index + 1];
+			}
+		}
+	}
+	return nullptr;
+}
+
+#undef LOCTEXT_NAMESPACE
diff --git a/Engine/Plugins/Web/WebAPI/Source/WebAPIEditor/Private/Details/Widgets/SWebAPISchemaTreeTableRow.h b/Engine/Plugins/Web/WebAPI/Source/WebAPIEditor/Private/Details/Widgets/SWebAPISchemaTreeTableRow.h
new file mode 100644
--- /dev/null
+++ b/Engine/Plugins/Web/WebAPI/Source/WebAPIEditor/Private/Details/Widgets/SWebAPISchemaTreeTableRow.h
@@ -0,0 +1,55 @@
+﻿// Copyright Epic Games, Inc. All Rights Reserved.
+
+#pragma once
+
+#include "CoreMinimal.h"
+#include "Widgets/DeclarativeSyntaxSupport.h"
+#include "Widgets/SWidget.h"
+#include "Widgets/Views/STableRow.h"
+#include "Widgets/Views/STableViewBase.h"
+
+class IWebAPISchemaTreeTableRow
+{
+};
+
+template <typename ItemType>
+class SWebAPISchemaTreeTableRow
+	: public STableRow<TSharedRef<ItemType>>
+	, public IWebAPISchemaTreeTableRow
+{
+public:
+	DECLARE_DELEGATE_OneParam(FOnColumnWidthChanged, float);
+	
+public:
+	SLATE_BEGIN_ARGS(SWebAPISchemaTreeTableRow)
+	{ }
+		SLATE_DEFAULT_SLOT(typename SWebAPISchemaTreeTableRow<ItemType>::FArguments, Content)
+	SLATE_END_ARGS()
+
+public:
+	void Construct(const FArguments& InArgs, const TSharedRef<ItemType>& InViewModel, const TSharedRef<STableViewBase>& InOwnerTableView);
+
+protected:
+	TSharedPtr<ItemType> ViewModel;
+};
+
+template <typename ItemType>
+void SWebAPISchemaTreeTableRow<ItemType>::Construct(const FArguments& InArgs, const TSharedRef<ItemType>& InViewModel, const TSharedRef<STableViewBase>& InOwnerTableView)
+{
+	ViewModel = InViewModel;
+	
+	STableRow<TSharedRef<ItemType>>::Construct(
+		SMultiColumnTableRow<TSharedRef<ItemType>>::FArguments()
+		.Padding(0)		
+		.Content()
+		[
+			// Wrap in border to use common Tooltip getter
+			SNew(SBorder)
+			.BorderBackgroundColor(FLinearColor(0,0,0,0))
+			.ToolTipText(InViewModel->GetTooltip())
+			[
+				InArgs._Content.Widget
+			]
+		],
+		InOwnerTableView);
+}
diff --git a/Engine/Plugins/Web/WebAPI/Source/WebAPIEditor/Private/Details/Widgets/SWebAPITreeView.cpp b/Engine/Plugins/Web/WebAPI/Source/WebAPIEditor/Private/Details/Widgets/SWebAPITreeView.cpp
new file mode 100644
--- /dev/null
+++ b/Engine/Plugins/Web/WebAPI/Source/WebAPIEditor/Private/Details/Widgets/SWebAPITreeView.cpp
@@ -0,0 +1,302 @@
+﻿// Copyright Epic Games, Inc. All Rights Reserved.
+
+#include "SWebAPITreeView.h"
+
+#include "Algo/ForEach.h"
+#include "Details/ViewModels/WebAPIViewModel.h"
+#include "Details/Widgets/SWebAPISchemaEnumRow.h"
+#include "Details/Widgets/SWebAPISchemaEnumValueRow.h"
+#include "Details/Widgets/SWebAPISchemaModelRow.h"
+#include "Details/Widgets/SWebAPISchemaOperationParameterRow.h"
+#include "Details/Widgets/SWebAPISchemaOperationRequestRow.h"
+#include "Details/Widgets/SWebAPISchemaOperationResponseRow.h"
+#include "Details/Widgets/SWebAPISchemaOperationRow.h"
+#include "Details/Widgets/SWebAPISchemaParameterRow.h"
+#include "Details/Widgets/SWebAPISchemaPropertyRow.h"
+#include "Details/Widgets/SWebAPISchemaServiceRow.h"
+#include "Details/Widgets/SWebAPISchemaTreeTableRow.h"
+#include "Widgets/Input/SSearchBox.h"
+#include "Widgets/Layout/SScrollBorder.h"
+#include "Widgets/Layout/SScrollBox.h"
+#include "Widgets/Text/STextBlock.h"
+#include "Widgets/Views/STreeView.h"
+
+#define LOCTEXT_NAMESPACE "SWebAPITreeView"
+
+/** Filter utility class, used by search. */
+class FWebAPITreeViewFilterContext : public ITextFilterExpressionContext
+{
+public:
+	explicit FWebAPITreeViewFilterContext(const FText& InString)
+		: String(InString)
+	{
+	}
+
+	virtual bool TestBasicStringExpression(const FTextFilterString& InValue, const ETextFilterTextComparisonMode InTextComparisonMode) const override
+	{
+		return TextFilterUtils::TestBasicStringExpression(String.ToString(), InValue, InTextComparisonMode);
+	}
+
+	virtual bool TestComplexExpression(const FName& InKey, const FTextFilterString& InValue, const ETextFilterComparisonOperation InComparisonOperation, const ETextFilterTextComparisonMode InTextComparisonMode) const override
+	{
+		return false;
+	}
+
+private:
+	const FText& String;
+};
+
+void SWebAPITreeView::Construct(const FArguments& InArgs, const FText& InLabel, const TArray<FTreeItemType>& InTreeViewItems)
+{
+	static const FName ColumnName("Name");
+
+	TreeViewItems = InTreeViewItems;
+	TextFilter = MakeShared<FTextFilterExpressionEvaluator>(ETextFilterExpressionEvaluatorMode::BasicString);
+	
+	TreeViewWidget =
+			SNew(STreeView<FTreeItemType>)
+			.SelectionMode(ESelectionMode::Multi)
+			.TreeItemsSource(&FilteredTreeViewItems)
+			.OnSelectionChanged(InArgs._OnSelectionChanged)
+			.OnGetChildren(this, &SWebAPITreeView::GetTreeItemChildren)
+			.OnGenerateRow(this, &SWebAPITreeView::GenerateRow)
+			.HeaderRow(
+				SNew(SHeaderRow)
+				.Visibility(EVisibility::Visible)
+				+SHeaderRow::Column(ColumnName)
+				.DefaultLabel(InLabel)
+			);
+
+	ChildSlot
+	[
+		SNew(SVerticalBox)
+
+		+ SVerticalBox::Slot()
+		.AutoHeight()
+		[
+			SNew(SSearchBox)
+			.OnTextChanged(this, &SWebAPITreeView::OnFilterTextChanged)
+		]
+
+		+ SVerticalBox::Slot()
+		.FillHeight(1.0f)
+		[
+			SNew(SScrollBorder, TreeViewWidget.ToSharedRef())
+			[
+				TreeViewWidget.ToSharedRef()
+			]
+		]
+	];
+}
+
+void SWebAPITreeView::Refresh(const TArray<TSharedPtr<IWebAPIViewModel>>& InTreeViewModels)
+{
+	if(!TreeViewWidget.IsValid())
+	{
+		return;
+	}
+
+	if(!InTreeViewModels.IsEmpty())
+	{
+		TreeViewItems = InTreeViewModels;
+	}
+
+	TUniqueFunction<void(const FTreeItemType&)> ExpandRecursive = [&](const FTreeItemType& InParent)
+	{
+		if(!InParent->ShouldExpandByDefault())
+		{
+			return;			
+		}
+		
+		TreeViewWidget->SetItemExpansion(InParent, true);
+		
+		TArray<TSharedPtr<IWebAPIViewModel>> Children;
+		InParent->GetChildren(Children);
+		Algo::ForEach(Children, [&ExpandRecursive](const FTreeItemType& InChild)
+		{
+			ExpandRecursive(InChild);
+		});
+	};
+
+	Algo::ForEach(TreeViewItems, [&ExpandRecursive](const FTreeItemType& InChild)
+	{
+		ExpandRecursive(InChild);
+	});
+
+	if (FilterText.IsEmpty())
+	{
+		FilteredTreeViewItems = TreeViewItems;
+	}
+	else
+	{
+		TArray<TSharedPtr<IWebAPIViewModel>> AllTreeViewItems;
+
+		// Flatten hierarchy and expand top-level items
+		for(const FTreeItemType& TreeViewItem : TreeViewItems)
+		{
+			TArray<TSharedPtr<IWebAPIViewModel>> TreeViewItemChildren;
+			TreeViewItem->GetChildren(AllTreeViewItems);
+
+			// Always expand first level
+			TreeViewWidget->SetItemExpansion(TreeViewItem, true);
+		}
+
+		for(const FTreeItemType& TreeViewItem : AllTreeViewItems)
+		{
+			if(TextFilter->TestTextFilter(FWebAPITreeViewFilterContext(TreeViewItem->GetLabel())))
+			{
+				FilteredTreeViewItems.Add(TreeViewItem);
+			}
+		}
+	}
+
+	TreeViewWidget->RequestTreeRefresh();
+}
+
+void SWebAPITreeView::GetTreeItemChildren(FTreeItemType InObject, TArray<FTreeItemType>& OutChildren)
+{
+	if(!InObject.IsValid() || !InObject->IsValid())
+	{
+		return;
+	}
+
+	InObject->GetChildren(OutChildren);
+}
+
+void SWebAPITreeView::OnFilterTextChanged(const FText& InSearchText)
+{
+	FilterText = InSearchText;
+	FilteredTreeViewItems.Reset();
+	TextFilter->SetFilterText(FilterText);
+ 
+	Refresh();
+}
+
+template <>
+TSharedRef<ITableRow> SWebAPITreeView::GenerateRow<IWebAPIViewModel>(TSharedPtr<IWebAPIViewModel> InViewModel, const TSharedRef<STableViewBase>& InOwnerTable)
+{
+	check(InViewModel.IsValid());
+
+	const FName ViewModelTypeName = InViewModel->GetViewModelTypeName();
+
+	if(ViewModelTypeName == IWebAPIViewModel::NAME_Service)
+	{
+		return GenerateRow(StaticCastSharedPtr<FWebAPIServiceViewModel>(InViewModel), InOwnerTable);
+	}
+	
+	if(ViewModelTypeName == IWebAPIViewModel::NAME_Operation)
+	{
+		return GenerateRow(StaticCastSharedPtr<FWebAPIOperationViewModel>(InViewModel), InOwnerTable);
+	}
+
+	if(ViewModelTypeName == IWebAPIViewModel::NAME_OperationRequest)
+	{
+		return GenerateRow(StaticCastSharedPtr<FWebAPIOperationRequestViewModel>(InViewModel), InOwnerTable);
+	}
+
+	if(ViewModelTypeName == IWebAPIViewModel::NAME_OperationParameter)
+	{
+		return GenerateRow(StaticCastSharedPtr<FWebAPIOperationParameterViewModel>(InViewModel), InOwnerTable);
+	}
+
+	if(ViewModelTypeName == IWebAPIViewModel::NAME_OperationResponse)
+	{
+		return GenerateRow(StaticCastSharedPtr<FWebAPIOperationResponseViewModel>(InViewModel), InOwnerTable);
+	}
+	
+	if(ViewModelTypeName == IWebAPIViewModel::NAME_Model)
+	{
+		return GenerateRow(StaticCastSharedPtr<FWebAPIModelViewModel>(InViewModel), InOwnerTable);
+	}
+	
+	if(ViewModelTypeName == IWebAPIViewModel::NAME_Property)
+	{
+		return GenerateRow(StaticCastSharedPtr<FWebAPIPropertyViewModel>(InViewModel), InOwnerTable);
+	}
+
+	if(ViewModelTypeName == IWebAPIViewModel::NAME_Enum)
+	{
+		return GenerateRow(StaticCastSharedPtr<FWebAPIEnumViewModel>(InViewModel), InOwnerTable);
+	}
+	
+	if(ViewModelTypeName == IWebAPIViewModel::NAME_EnumValue)
+	{
+		return GenerateRow(StaticCastSharedPtr<FWebAPIEnumValueViewModel>(InViewModel), InOwnerTable);
+	}
+
+	if(ViewModelTypeName == IWebAPIViewModel::NAME_Parameter)
+	{
+		return GenerateRow<FWebAPIParameterViewModel>(StaticCastSharedPtr<FWebAPIParameterViewModel>(InViewModel), InOwnerTable);
+	}
+
+	TSharedRef<STableRow<FTreeItemType>> Row =
+		SNew(STableRow<FTreeItemType>, InOwnerTable)
+		[
+			SNew(STextBlock)
+			.Text(InViewModel->GetLabel())
+		];
+
+	return Row;
+}
+
+template <>
+TSharedRef<ITableRow> SWebAPITreeView::GenerateRow<FWebAPIServiceViewModel>(TSharedPtr<FWebAPIServiceViewModel> InViewModel, const TSharedRef<STableViewBase>& InOwnerTable)
+{
+	return SNew(SWebAPISchemaServiceRow, InViewModel.ToSharedRef(), InOwnerTable);
+}
+
+template <>
+TSharedRef<ITableRow> SWebAPITreeView::GenerateRow<FWebAPIOperationViewModel>(TSharedPtr<FWebAPIOperationViewModel> InViewModel, const TSharedRef<STableViewBase>& InOwnerTable)
+{
+	return SNew(SWebAPISchemaOperationRow, InViewModel.ToSharedRef(), InOwnerTable);
+}
+
+template <>
+TSharedRef<ITableRow> SWebAPITreeView::GenerateRow<FWebAPIOperationRequestViewModel>(TSharedPtr<FWebAPIOperationRequestViewModel> InViewModel, const TSharedRef<STableViewBase>& InOwnerTable)
+{
+	return SNew(SWebAPISchemaOperationRequestRow, InViewModel.ToSharedRef(), InOwnerTable);
+}
+
+template <>
+TSharedRef<ITableRow> SWebAPITreeView::GenerateRow<FWebAPIOperationParameterViewModel>(TSharedPtr<FWebAPIOperationParameterViewModel> InViewModel, const TSharedRef<STableViewBase>& InOwnerTable)
+{
+	return SNew(SWebAPISchemaOperationParameterRow, InViewModel.ToSharedRef(), InOwnerTable);
+}
+
+template <>
+TSharedRef<ITableRow> SWebAPITreeView::GenerateRow<FWebAPIOperationResponseViewModel>(TSharedPtr<FWebAPIOperationResponseViewModel> InViewModel, const TSharedRef<STableViewBase>& InOwnerTable)
+{
+	return SNew(SWebAPISchemaOperationResponseRow, InViewModel.ToSharedRef(), InOwnerTable);
+}
+
+template <>
+TSharedRef<ITableRow> SWebAPITreeView::GenerateRow<FWebAPIModelViewModel>(TSharedPtr<FWebAPIModelViewModel> InViewModel, const TSharedRef<STableViewBase>& InOwnerTable)
+{
+	return SNew(SWebAPISchemaModelRow, InViewModel.ToSharedRef(), InOwnerTable);
+}
+
+template <>
+TSharedRef<ITableRow> SWebAPITreeView::GenerateRow<FWebAPIPropertyViewModel>(TSharedPtr<FWebAPIPropertyViewModel> InViewModel, const TSharedRef<STableViewBase>& InOwnerTable)
+{
+	return SNew(SWebAPISchemaPropertyRow, InViewModel.ToSharedRef(), InOwnerTable);
+}
+
+template <>
+TSharedRef<ITableRow> SWebAPITreeView::GenerateRow<FWebAPIEnumViewModel>(TSharedPtr<FWebAPIEnumViewModel> InViewModel, const TSharedRef<STableViewBase>& InOwnerTable)
+{
+	return SNew(SWebAPISchemaEnumRow, InViewModel.ToSharedRef(), InOwnerTable);
+}
+
+template <>
+TSharedRef<ITableRow> SWebAPITreeView::GenerateRow<FWebAPIEnumValueViewModel>(TSharedPtr<FWebAPIEnumValueViewModel> InViewModel, const TSharedRef<STableViewBase>& InOwnerTable)
+{
+	return SNew(SWebAPISchemaEnumValueRow, InViewModel.ToSharedRef(), InOwnerTable);
+}
+
+template <>
+TSharedRef<ITableRow> SWebAPITreeView::GenerateRow<FWebAPIParameterViewModel>(TSharedPtr<FWebAPIParameterViewModel> InViewModel, const TSharedRef<STableViewBase>& InOwnerTable)
+{
+	return SNew(SWebAPISchemaParameterRow, InViewModel.ToSharedRef(), InOwnerTable);
+}
+
+#undef LOCTEXT_NAMESPACE
diff --git a/Engine/Plugins/Web/WebAPI/Source/WebAPIEditor/Private/Details/Widgets/SWebAPITreeView.h b/Engine/Plugins/Web/WebAPI/Source/WebAPIEditor/Private/Details/Widgets/SWebAPITreeView.h
new file mode 100644
--- /dev/null
+++ b/Engine/Plugins/Web/WebAPI/Source/WebAPIEditor/Private/Details/Widgets/SWebAPITreeView.h
@@ -0,0 +1,116 @@
+﻿// Copyright Epic Games, Inc. All Rights Reserved.
+
+#pragma once
+
+#include "CoreMinimal.h"
+#include "Algo/Transform.h"
+#include "Details/ViewModels/WebAPIEnumViewModel.h"
+#include "Details/ViewModels/WebAPIModelViewModel.h"
+#include "Details/ViewModels/WebAPIOperationParameterViewModel.h"
+#include "Details/ViewModels/WebAPIOperationRequestViewModel.h"
+#include "Details/ViewModels/WebAPIOperationResponseViewModel.h"
+#include "Details/ViewModels/WebAPIParameterViewModel.h"
+#include "Details/ViewModels/WebAPIServiceViewModel.h"
+#include "Details/ViewModels/WebAPIViewModel.h"
+#include "Misc/TextFilterExpressionEvaluator.h"
+#include "Widgets/SCompoundWidget.h"
+#include "Widgets/Views/ITableRow.h"
+#include "Widgets/Views/STreeView.h"
+
+class IDetailLayoutBuilder;
+
+/** A widget containing either a tree of services or models. */
+class WEBAPIEDITOR_API SWebAPITreeView : public SCompoundWidget
+{
+public:
+	using FTreeItemType = TSharedPtr<IWebAPIViewModel>;
+
+	using FOnSelectionChanged  = typename TSlateDelegates<FTreeItemType>::FOnSelectionChanged;
+	
+	SLATE_BEGIN_ARGS(SWebAPITreeView)
+		: _OnSelectionChanged()
+	{ }
+		SLATE_EVENT(FOnSelectionChanged, OnSelectionChanged)
+	SLATE_END_ARGS()
+
+	/** Constructs this widget with InArgs */
+	void Construct(const FArguments& InArgs, const FText& InLabel, const TArray<FTreeItemType>& InTreeViewItems = {});
+
+	/** Refresh the tree view. If new items are provided, they'll be replaced for this widget. */
+	void Refresh(const TArray<TSharedPtr<IWebAPIViewModel>>& InTreeViewModels = {});
+
+	/** Refresh the tree view. If new items are provided, they'll be replaced for this widget. */
+	template <typename ViewModelType>
+	void Refresh(const TArray<TSharedPtr<ViewModelType>>& InTreeViewModels = {});
+
+private:
+	/** Appends children of the provided Object to OutChildren. */
+	void GetTreeItemChildren(FTreeItemType InObject, TArray<FTreeItemType>& OutChildren);
+
+	/** Called when the associated search input changes. */
+	void OnFilterTextChanged(const FText& InSearchText);
+
+	/** Makes the row widget for the provided ViewModel. */
+	template <typename ViewModelType>
+	TSharedRef<ITableRow> GenerateRow(TSharedPtr<ViewModelType> InViewModel, const TSharedRef<STableViewBase>& InOwnerTable);
+
+	template <>
+	TSharedRef<ITableRow> GenerateRow(TSharedPtr<IWebAPIViewModel> InViewModel, const TSharedRef<STableViewBase>& InOwnerTable);
+
+	template <>
+	TSharedRef<ITableRow> GenerateRow(TSharedPtr<FWebAPIServiceViewModel> InViewModel, const TSharedRef<STableViewBase>& InOwnerTable);
+
+	template <>
+	TSharedRef<ITableRow> GenerateRow(TSharedPtr<FWebAPIOperationViewModel> InViewModel, const TSharedRef<STableViewBase>& InOwnerTable);
+
+	template <>
+	TSharedRef<ITableRow> GenerateRow(TSharedPtr<FWebAPIOperationRequestViewModel> InViewModel, const TSharedRef<STableViewBase>& InOwnerTable);
+
+	template <>
+	TSharedRef<ITableRow> GenerateRow(TSharedPtr<FWebAPIOperationParameterViewModel> InViewModel, const TSharedRef<STableViewBase>& InOwnerTable);
+
+	template <>
+	TSharedRef<ITableRow> GenerateRow(TSharedPtr<FWebAPIOperationResponseViewModel> InViewModel, const TSharedRef<STableViewBase>& InOwnerTable);
+	
+	template <>
+	TSharedRef<ITableRow> GenerateRow(TSharedPtr<FWebAPIModelViewModel> InViewModel, const TSharedRef<STableViewBase>& InOwnerTable);
+
+	template <>
+	TSharedRef<ITableRow> GenerateRow(TSharedPtr<FWebAPIPropertyViewModel> InViewModel, const TSharedRef<STableViewBase>& InOwnerTable);
+
+	template <>
+	TSharedRef<ITableRow> GenerateRow(TSharedPtr<FWebAPIEnumViewModel> InViewModel, const TSharedRef<STableViewBase>& InOwnerTable);
+
+	template <>
+	TSharedRef<ITableRow> GenerateRow(TSharedPtr<FWebAPIEnumValueViewModel> InViewModel, const TSharedRef<STableViewBase>& InOwnerTable);
+
+	template <>
+	TSharedRef<ITableRow> GenerateRow(TSharedPtr<FWebAPIParameterViewModel> InViewModel, const TSharedRef<STableViewBase>& InOwnerTable);
+
+private:
+	/** TreeView Widget. */
+	TSharedPtr<STreeView<FTreeItemType>> TreeViewWidget;
+
+	/** TreeView model context. */
+	TArray<FTreeItemType> TreeViewItems;
+
+	/** TreeView model context, filtered according to FilterText. */
+	TArray<FTreeItemType> FilteredTreeViewItems;
+
+	/** Current filter text, if any. */
+	FText FilterText;
+
+	/** The filter itself, if any. */
+	TSharedPtr<FTextFilterExpressionEvaluator> TextFilter;
+};
+
+template <typename ViewModelType>
+void SWebAPITreeView::Refresh(const TArray<TSharedPtr<ViewModelType>>& InTreeViewModels)
+{
+	TArray<TSharedPtr<IWebAPIViewModel>> CastTreeViewModels;
+	Algo::Transform(InTreeViewModels, CastTreeViewModels, [](const TSharedPtr<ViewModelType>& InItem)
+	{
+		return InItem;
+	});
+	Refresh(MoveTemp(CastTreeViewModels));
+}
diff --git a/Engine/Plugins/Web/WebAPI/Source/WebAPIEditor/Private/Tests/WebAPIEditorUtilities.spec.cpp b/Engine/Plugins/Web/WebAPI/Source/WebAPIEditor/Private/Tests/WebAPIEditorUtilities.spec.cpp
new file mode 100644
--- /dev/null
+++ b/Engine/Plugins/Web/WebAPI/Source/WebAPIEditor/Private/Tests/WebAPIEditorUtilities.spec.cpp
@@ -0,0 +1,32 @@
+﻿// Copyright Epic Games, Inc. All Rights Reserved.
+
+#include "Misc/AutomationTest.h"
+
+#if WITH_AUTOMATION_TESTS
+
+#include "WebAPIEditorUtilities.h"
+
+using namespace UE::WebAPI;
+
+BEGIN_DEFINE_SPEC(FWebAPIEditorUtilitiesSpec,
+	TEXT("Plugin.WebAPI.Editor.Utilities"),
+	EAutomationTestFlags::EditorContext | EAutomationTestFlags::EngineFilter | EAutomationTestFlags::ApplicationContextMask)
+END_DEFINE_SPEC(FWebAPIEditorUtilitiesSpec)
+
+void FWebAPIEditorUtilitiesSpec::Define()
+{
+	Describe(TEXT("ToPascalCase"), [this]
+	{
+		It(TEXT("Strips leading illegal tokens"), [this]
+		{
+			const FStringView TestString = TEXT("$teststring");
+			const FString ExpectedString = TEXT("teststring");
+
+			const FString ResultString = UE::WebAPI::FWebAPIStringUtilities::Get()->ToPascalCase(TestString);
+			
+			TestEqual("String doesn't contain illegal characters", ResultString, ExpectedString);
+		});
+	});
+}
+
+#endif
diff --git a/Engine/Plugins/Web/WebAPI/Source/WebAPIEditor/Private/WebAPIPluginWizardDefinition.cpp b/Engine/Plugins/Web/WebAPI/Source/WebAPIEditor/Private/WebAPIPluginWizardDefinition.cpp
new file mode 100644
--- /dev/null
+++ b/Engine/Plugins/Web/WebAPI/Source/WebAPIEditor/Private/WebAPIPluginWizardDefinition.cpp
@@ -0,0 +1,169 @@
+﻿// Copyright Epic Games, Inc. All Rights Reserved.
+
+#include "WebAPIPluginWizardDefinition.h"
+
+#include "IPluginBrowser.h"
+#include "WebAPIEditorSubsystem.h"
+#include "HAL/PlatformFileManager.h"
+#include "Interfaces/IPluginManager.h"
+#include "Misc/Paths.h"
+
+#define LOCTEXT_NAMESPACE "WebAPIPluginWizardDefinition"
+
+/** A plugin template description that applies template-specific logic. */
+struct FWebAPIGeneratedPluginTemplateDescription : public FPluginTemplateDescription
+{
+	FWebAPIGeneratedPluginTemplateDescription(FText InName, FText InDescription, FString InOnDiskPath, bool bInCanContainContent = true, EHostType::Type InModuleDescriptorType = EHostType::Runtime, ELoadingPhase::Type InLoadingPhase = ELoadingPhase::Default)
+		: FPluginTemplateDescription(InName, InDescription, InOnDiskPath, bInCanContainContent, InModuleDescriptorType, InLoadingPhase)
+	{
+		bCanBePlacedInEngine = false;
+	}
+
+	virtual void CustomizeDescriptorBeforeCreation(FPluginDescriptor& Descriptor) override
+	{
+		// Rename module with "Generated" suffix
+		FModuleDescriptor& ModuleDescriptor = Descriptor.Modules[0];
+		ModuleDescriptor.Name = FName(ModuleDescriptor.Name.ToString() + TEXT("Generated"));
+
+		// Add dependency on the WebAPI plugin
+		Descriptor.Plugins.Add(FPluginReferenceDescriptor(TEXT("WebAPI"), true));
+	}
+};
+
+FWebAPIPluginWizardDefinition::FWebAPIPluginWizardDefinition(bool bInIsContentOnlyProject)
+	: bIsContentOnlyProject(bInIsContentOnlyProject)
+{
+	PluginBaseDir = IPluginManager::Get().FindPlugin(TEXT("WebAPI"))->GetBaseDir();
+
+	PopulateTemplatesSource();
+}
+
+void FWebAPIPluginWizardDefinition::PopulateTemplatesSource()
+{
+	const FText GeneratedAPITemplateName = LOCTEXT("GeneratedAPILabel", "Generated API");	
+	const FText GeneratedAPIDescription = LOCTEXT("GeneratedAPITemplateDesc", "Create a plugin to contain code generated by a WebAPI definition.");
+
+	const TSharedRef<FPluginTemplateDescription> GeneratedAPITemplate = MakeShared<FWebAPIGeneratedPluginTemplateDescription>(GeneratedAPITemplateName, GeneratedAPIDescription, PluginBaseDir / TEXT("Templates") / TEXT("GeneratedAPI"), true, EHostType::Runtime);
+	TemplateDefinitions.Add(GeneratedAPITemplate);
+
+	TemplateDefinitions.Sort([](const TSharedRef<FPluginTemplateDescription>& A, const TSharedRef<FPluginTemplateDescription>& B)
+	{
+		return A->Name.CompareTo(B->Name) <= 0;
+	});
+}
+
+const TArray<TSharedRef<FPluginTemplateDescription>>& FWebAPIPluginWizardDefinition::GetTemplatesSource() const
+{
+	return TemplateDefinitions;
+}
+
+void FWebAPIPluginWizardDefinition::OnTemplateSelectionChanged(TSharedPtr<FPluginTemplateDescription> InSelectedItem, ESelectInfo::Type SelectInfo)
+{
+	CurrentTemplateDefinition = InSelectedItem;
+}
+
+TSharedPtr<FPluginTemplateDescription> FWebAPIPluginWizardDefinition::GetSelectedTemplate() const
+{
+	return CurrentTemplateDefinition;
+}
+
+bool FWebAPIPluginWizardDefinition::HasValidTemplateSelection() const
+{
+	return CurrentTemplateDefinition.IsValid();
+}
+
+void FWebAPIPluginWizardDefinition::ClearTemplateSelection()
+{
+	CurrentTemplateDefinition.Reset();
+}
+
+bool FWebAPIPluginWizardDefinition::HasModules() const
+{
+	const FString SourceFolderPath = GetPluginFolderPath() / TEXT("Source");
+	
+	return FPaths::DirectoryExists(SourceFolderPath);
+}
+
+bool FWebAPIPluginWizardDefinition::IsMod() const
+{
+	return false;
+}
+
+FText FWebAPIPluginWizardDefinition::GetInstructions() const
+{
+	return LOCTEXT("ChoosePluginTemplate", "Choose a template and then specify a name to create a new plugin.");
+}
+
+bool FWebAPIPluginWizardDefinition::GetPluginIconPath(FString& OutIconPath) const
+{
+	return GetTemplateIconPath(CurrentTemplateDefinition.ToSharedRef(), OutIconPath);
+}
+
+EHostType::Type FWebAPIPluginWizardDefinition::GetPluginModuleDescriptor() const
+{
+	EHostType::Type ModuleDescriptorType = EHostType::Runtime;
+
+	if (CurrentTemplateDefinition.IsValid())
+	{
+		ModuleDescriptorType = CurrentTemplateDefinition->ModuleDescriptorType;
+	}
+
+	return ModuleDescriptorType;
+}
+
+ELoadingPhase::Type FWebAPIPluginWizardDefinition::GetPluginLoadingPhase() const
+{
+	ELoadingPhase::Type Phase = ELoadingPhase::Default;
+
+	if (CurrentTemplateDefinition.IsValid())
+	{
+		Phase = CurrentTemplateDefinition->LoadingPhase;
+	}
+
+	return Phase;
+}
+
+bool FWebAPIPluginWizardDefinition::GetTemplateIconPath(TSharedRef<FPluginTemplateDescription> Template, FString& OutIconPath) const
+{
+	bool bRequiresDefaultIcon = false;
+
+	const FString TemplateFolderName = GetFolderForTemplate(Template);
+
+	OutIconPath = TemplateFolderName / TEXT("Resources/Icon128.png");
+	if (!FPlatformFileManager::Get().GetPlatformFile().FileExists(*OutIconPath))
+	{
+		OutIconPath = PluginBaseDir / TEXT("Resources/DefaultIcon128.png");
+		bRequiresDefaultIcon = true;
+	}
+
+	return bRequiresDefaultIcon;
+}
+
+TArray<FString> FWebAPIPluginWizardDefinition::GetFoldersForSelection() const
+{
+	TArray<FString> SelectedFolders;
+
+	if (CurrentTemplateDefinition.IsValid())
+	{
+		SelectedFolders.Add(GetFolderForTemplate(CurrentTemplateDefinition.ToSharedRef()));
+	}
+
+	return SelectedFolders;
+}
+
+void FWebAPIPluginWizardDefinition::PluginCreated(const FString& PluginName, bool bWasSuccessful) const
+{
+	GEditor->GetEditorSubsystem<UWebAPIEditorSubsystem>()->OnModuleCreated().Broadcast(PluginName);
+}
+
+FString FWebAPIPluginWizardDefinition::GetPluginFolderPath() const
+{
+	return GetFolderForTemplate(CurrentTemplateDefinition.ToSharedRef());
+}
+
+FString FWebAPIPluginWizardDefinition::GetFolderForTemplate(TSharedRef<FPluginTemplateDescription> InTemplate) const
+{
+	return InTemplate->OnDiskPath;
+}
+
+#undef LOCTEXT_NAMESPACE
\ No newline at end of file
diff --git a/Engine/Plugins/Web/WebAPI/Source/WebAPIEditor/Private/WebAPIPluginWizardDefinition.h b/Engine/Plugins/Web/WebAPI/Source/WebAPIEditor/Private/WebAPIPluginWizardDefinition.h
new file mode 100644
--- /dev/null
+++ b/Engine/Plugins/Web/WebAPI/Source/WebAPIEditor/Private/WebAPIPluginWizardDefinition.h
@@ -0,0 +1,59 @@
+﻿// Copyright Epic Games, Inc. All Rights Reserved.
+
+#pragma once
+
+#include "CoreMinimal.h"
+#include "IPluginWizardDefinition.h"
+
+/**
+ * New plugin wizard definition used by PluginBrowser
+ */
+class FWebAPIPluginWizardDefinition : public IPluginWizardDefinition
+{
+public:
+	FWebAPIPluginWizardDefinition(bool bIsContentOnlyProject = false);
+
+	// Begin IPluginWizardDefinition interface
+	virtual const TArray<TSharedRef<FPluginTemplateDescription>>& GetTemplatesSource() const override;
+	virtual void OnTemplateSelectionChanged(TSharedPtr<FPluginTemplateDescription> InSelectedItem, ESelectInfo::Type SelectInfo) override;
+	virtual TSharedPtr<FPluginTemplateDescription> GetSelectedTemplate() const override;
+	virtual void ClearTemplateSelection() override;
+	virtual bool HasValidTemplateSelection() const override;
+
+	virtual bool CanShowOnStartup() const override { return false; }
+	virtual bool HasModules() const override;
+	virtual bool IsMod() const override;
+	virtual void OnShowOnStartupCheckboxChanged(ECheckBoxState CheckBoxState) override {}
+	virtual ECheckBoxState GetShowOnStartupCheckBoxState() const override { return ECheckBoxState::Undetermined; }
+	virtual TSharedPtr<class SWidget> GetCustomHeaderWidget() override { return nullptr; }
+	virtual FText GetInstructions() const override;
+
+	virtual bool GetPluginIconPath(FString& OutIconPath) const override;
+	virtual EHostType::Type GetPluginModuleDescriptor() const override;
+	virtual ELoadingPhase::Type GetPluginLoadingPhase() const override;
+	virtual bool GetTemplateIconPath(TSharedRef<FPluginTemplateDescription> InTemplate, FString& OutIconPath) const override;
+	virtual FString GetPluginFolderPath() const override;
+	virtual TArray<FString> GetFoldersForSelection() const override;
+	virtual void PluginCreated(const FString& PluginName, bool bWasSuccessful) const override;
+	// End IPluginWizardDefinition interface
+
+private:
+	/** Creates the templates that can be used by the plugin manager to generate the plugin */
+	void PopulateTemplatesSource();
+
+	/** Gets the folder for the specified template. */
+	FString GetFolderForTemplate(TSharedRef<FPluginTemplateDescription> InTemplate) const;
+
+private:
+	/** The templates available to this definition */
+	TArray<TSharedRef<FPluginTemplateDescription>> TemplateDefinitions;
+
+	/** The currently selected template definition */
+	TSharedPtr<FPluginTemplateDescription> CurrentTemplateDefinition;
+
+	/** Base directory of the plugin templates */
+	FString PluginBaseDir;
+
+	/** If true, this definition is for a project that can only contain content */
+	bool bIsContentOnlyProject;
+};
diff --git a/Engine/Plugins/Web/WebAPI/Source/WebAPILiquidJS/Private/WebAPILiquidJSCodeGenerator.cpp b/Engine/Plugins/Web/WebAPI/Source/WebAPILiquidJS/Private/WebAPILiquidJSCodeGenerator.cpp
new file mode 100644
--- /dev/null
+++ b/Engine/Plugins/Web/WebAPI/Source/WebAPILiquidJS/Private/WebAPILiquidJSCodeGenerator.cpp
@@ -0,0 +1,825 @@
+﻿// Copyright Epic Games, Inc. All Rights Reserved.
+
+#include "WebAPILiquidJSCodeGenerator.h"
+
+#include "HttpModule.h"
+#include "WebAPILiquidJSModule.h"
+#include "WebAPILiquidJSSettings.h"
+#include "Algo/ForEach.h"
+#include "Algo/Transform.h"
+#include "Async/Async.h"
+#include "CodeGen/WebAPICodeGenerator.h"
+#include "CodeGen/Dom/WebAPICodeGenFile.h"
+#include "CodeGen/Dom/WebAPICodeGenFunction.h"
+#include "CodeGen/Dom/WebAPICodeGenSettings.h"
+#include "Dom/WebAPIType.h"
+#include "Interfaces/IHttpResponse.h"
+#include "Misc/FileHelper.h"
+#include "Modules/ModuleManager.h"
+#include "Policies/CondensedJsonPrintPolicy.h"
+#include "Serialization/JsonSerializer.h"
+
+#define LOCTEXT_NAMESPACE "WebAPILiquidJSCodeGenerator"
+
+namespace UE
+{
+	namespace WebAPI
+	{
+		namespace Generator
+		{
+			namespace LiquidJS
+			{
+				/** Various functionality for converting the provided CodeGen objects to implementation-specific Json objects. */
+				namespace Private
+				{
+					template <typename CodeGenType>
+					TSharedPtr<FJsonObject> ToJson(const TSharedPtr<CodeGenType>& InCodeGenObject);
+
+					template <typename CodeGenType>
+					TSharedPtr<FJsonObject> ToJson(const CodeGenType& InCodeGenObject);
+
+					template <typename CodeGenType>
+					void ToJson(const TArray<TSharedPtr<CodeGenType>>& InCodeGenObjects, TArray<TSharedPtr<FJsonValue>>& OutJson)
+					{
+						OutJson.Reserve(InCodeGenObjects.Num());
+						for(const TSharedPtr<CodeGenType>& Element : InCodeGenObjects)
+						{
+							TSharedPtr<FJsonObject> ElementJson = ToJson(Element);
+							OutJson.Add(MakeShared<FJsonValueObject>(ElementJson));
+						}
+					}
+
+					template <>
+					void ToJson<FWebAPICodeGenBase>(const TArray<TSharedPtr<FWebAPICodeGenBase>>& InCodeGenObjects, TArray<TSharedPtr<FJsonValue>>& OutJson);
+
+					void ToJson(const TSet<FString>& InArray, TArray<TSharedPtr<FJsonValue>>& OutJson)
+					{
+						OutJson.Reserve(InArray.Num());
+						for(const FString& Element : InArray)
+						{
+							TSharedPtr<FJsonValueString> ElementJson = MakeShared<FJsonValueString>(Element);
+							OutJson.Add(MoveTemp(ElementJson));
+						}
+					}
+
+					void ToJson(const TArray<FString>& InArray, TArray<TSharedPtr<FJsonValue>>& OutJson)
+					{
+						OutJson.Reserve(InArray.Num());
+						for(const FString& Element : InArray)
+						{
+							TSharedPtr<FJsonValueString> ElementJson = MakeShared<FJsonValueString>(Element);
+							OutJson.Add(MoveTemp(ElementJson));
+						}
+					}
+
+					/** EnumValue Array ToJson */
+					void ToJson(const TArray<FWebAPICodeGenEnumValue>& InArray, TArray<TSharedPtr<FJsonValue>>& OutJson)
+					{
+						OutJson.Reserve(InArray.Num());
+						for(const FWebAPICodeGenEnumValue& Element : InArray)
+						{
+							TSharedPtr<FJsonObject> ElementJson = MakeShared<FJsonObject>();
+							ElementJson->SetStringField(TEXT("name"), Element.Name);
+							ElementJson->SetStringField(TEXT("displayName"), Element.DisplayName);
+							ElementJson->SetStringField(TEXT("jsonName"), Element.JsonName);
+							ElementJson->SetStringField(TEXT("description"), Element.Description);
+							if(Element.IntegralValue >= 0)
+							{
+								ElementJson->SetNumberField(TEXT("explicitValue"), Element.IntegralValue);
+							}							
+							OutJson.Add(MakeShared<FJsonValueObject>(ElementJson));
+						}
+					}
+
+					void ToJson(const TMap<FString, FString>& InMap, TArray<TSharedPtr<FJsonValue>>& OutJson)
+					{
+						OutJson.Reserve(InMap.Num());
+						for(const TPair<FString, FString>& Element : InMap)
+						{
+							const TSharedPtr<FJsonObject> ElementJson = MakeShared<FJsonObject>();
+							ElementJson->SetStringField(TEXT("key"), Element.Key);
+							if(!Element.Value.IsEmpty())
+							{
+								ElementJson->SetStringField(TEXT("value"), Element.Value);
+							}
+							OutJson.Add(MakeShared<FJsonValueObject>(ElementJson));
+						}
+					}
+
+					/** TypeNameVariant ToJson */
+					template <>
+					TSharedPtr<FJsonObject> ToJson(const FWebAPITypeNameVariant& InCodeGenObject)
+					{
+						check(InCodeGenObject.IsValid());
+
+						TSharedPtr<FJsonObject> JsonObject = MakeShared<FJsonObject>();
+						if(InCodeGenObject.HasTypeInfo())
+						{
+							JsonObject->SetStringField(TEXT("name"), InCodeGenObject.TypeInfo->Name);
+							JsonObject->SetStringField(TEXT("jsonName"), InCodeGenObject.TypeInfo->JsonName);
+							JsonObject->SetStringField(TEXT("jsonType"), InCodeGenObject.TypeInfo->JsonType.ToString());
+							JsonObject->SetStringField(TEXT("jsonPropertyToSerialize"), InCodeGenObject.TypeInfo->JsonPropertyToSerialize);
+							JsonObject->SetStringField(TEXT("printFormatSpecifier"), InCodeGenObject.TypeInfo->PrintFormatSpecifier);
+							JsonObject->SetStringField(TEXT("printFormatExpression"), InCodeGenObject.TypeInfo->PrintFormatExpression);
+							JsonObject->SetStringField(TEXT("namespace"), InCodeGenObject.TypeInfo->Namespace);
+							JsonObject->SetStringField(TEXT("prefix"), InCodeGenObject.TypeInfo->Prefix);
+							JsonObject->SetStringField(TEXT("declarationType"), InCodeGenObject.TypeInfo->DeclarationType.ToString());
+							JsonObject->SetBoolField(TEXT("isBuiltinType"), InCodeGenObject.TypeInfo->bIsBuiltinType);
+						}
+
+						return JsonObject;
+					}
+
+					/** NameVariant ToJson */
+					template <>
+					TSharedPtr<FJsonObject> ToJson(const FWebAPINameVariant& InCodeGenObject)
+					{
+						check(InCodeGenObject.IsValid());
+
+						TSharedPtr<FJsonObject> JsonObject = MakeShared<FJsonObject>();
+						if(InCodeGenObject.HasNameInfo())
+						{
+							JsonObject->SetStringField(TEXT("name"), InCodeGenObject.NameInfo.Name);
+							JsonObject->SetStringField(TEXT("jsonName"), InCodeGenObject.NameInfo.JsonName);
+							JsonObject->SetStringField(TEXT("prefix"), InCodeGenObject.NameInfo.Prefix);
+						}
+
+						return JsonObject;
+					}
+
+					/** Base ToJson */
+					template <>
+					TSharedPtr<FJsonObject> ToJson(const TSharedPtr<FWebAPICodeGenBase>& InCodeGenObject)
+					{
+						check(InCodeGenObject);
+
+						TSharedPtr<FJsonObject> JsonObject = MakeShared<FJsonObject>();
+						
+						JsonObject->SetStringField(TEXT("description"), InCodeGenObject->Description);
+						JsonObject->SetStringField(TEXT("module"), InCodeGenObject->Module);
+						JsonObject->SetStringField(TEXT("namespace"), InCodeGenObject->Namespace);
+						
+						TArray<TSharedPtr<FJsonValue>> SpecifiersJson;
+						ToJson(InCodeGenObject->Specifiers, SpecifiersJson);
+						JsonObject->SetArrayField(TEXT("specifiers"), SpecifiersJson);
+
+						TArray<TSharedPtr<FJsonValue>> MetadataJson;
+						ToJson(InCodeGenObject->Metadata, MetadataJson);
+						JsonObject->SetArrayField(TEXT("metadata"), MetadataJson);
+
+						return JsonObject;
+					}
+
+					/** Property ToJson */
+					template <>
+					TSharedPtr<FJsonObject> ToJson(const TSharedPtr<FWebAPICodeGenProperty>& InCodeGenObject)
+					{
+						check(InCodeGenObject);
+
+						const TSharedPtr<FJsonObject> JsonObject = ToJson<FWebAPICodeGenBase>(StaticCastSharedPtr<FWebAPICodeGenBase>(InCodeGenObject));
+
+						if(InCodeGenObject->Name.HasNameInfo())
+						{
+							const TSharedPtr<FJsonObject> NameJson = ToJson(InCodeGenObject->Name);
+							JsonObject->SetObjectField(TEXT("name"), NameJson);
+						}
+						else
+						{
+							JsonObject->SetStringField(TEXT("name"), InCodeGenObject->Name.ToString(true));
+						}
+
+						if(InCodeGenObject->Type.HasTypeInfo())
+						{
+							const TSharedPtr<FJsonObject> TypeJson = ToJson(InCodeGenObject->Type);
+							JsonObject->SetObjectField(TEXT("type"), TypeJson);
+
+							FString DefaultValue = InCodeGenObject->Type.TypeInfo->DefaultValue;
+							// Wrap in array initializer, if it's an array
+							if(InCodeGenObject->bIsArray)
+							{
+								DefaultValue = TEXT("{ }");
+							}
+							JsonObject->SetStringField(TEXT("defaultValue"), DefaultValue);
+						}
+						else
+						{
+							JsonObject->SetStringField(TEXT("type"), InCodeGenObject->Type.ToString(true));
+						}
+
+						if(!InCodeGenObject->DefaultValue.IsEmpty())
+						{
+							JsonObject->SetStringField(TEXT("defaultValue"), InCodeGenObject->DefaultValue);
+						}
+
+						JsonObject->SetBoolField(TEXT("isRequired"), InCodeGenObject->bIsRequired);
+						JsonObject->SetBoolField(TEXT("isArray"), InCodeGenObject->bIsArray);
+						JsonObject->SetBoolField(TEXT("isMixin"), InCodeGenObject->bIsMixin);
+
+						return JsonObject;
+					}
+
+					/** Enum ToJson */
+					template <>
+					TSharedPtr<FJsonObject> ToJson(const TSharedPtr<FWebAPICodeGenEnum>& InCodeGenObject)
+					{
+						check(InCodeGenObject);
+
+						const TSharedPtr<FJsonObject> JsonObject = ToJson<FWebAPICodeGenBase>(StaticCastSharedPtr<FWebAPICodeGenBase>(InCodeGenObject));
+						
+						if(InCodeGenObject->Name.HasTypeInfo())
+						{
+							const TSharedPtr<FJsonObject> NameJson = ToJson(InCodeGenObject->Name);
+							JsonObject->SetObjectField(TEXT("name"), NameJson);
+						}
+						else
+						{
+							JsonObject->SetStringField(TEXT("name"), InCodeGenObject->Name.ToString(true));
+						}
+
+						TArray<TSharedPtr<FJsonValue>> ValuesJson;
+						ToJson(InCodeGenObject->Values, ValuesJson);
+						JsonObject->SetArrayField(TEXT("values"), ValuesJson);
+
+						return JsonObject;
+					}
+
+					/** Struct ToJson */
+					template <>
+					TSharedPtr<FJsonObject> ToJson(const TSharedPtr<FWebAPICodeGenStruct>& InCodeGenObject)
+					{
+						check(InCodeGenObject);
+
+						const TSharedPtr<FJsonObject> JsonObject = ToJson<FWebAPICodeGenBase>(StaticCastSharedPtr<FWebAPICodeGenBase>(InCodeGenObject));
+
+						if(InCodeGenObject->Name.HasTypeInfo())
+						{
+							const TSharedPtr<FJsonObject> NameJson = ToJson(InCodeGenObject->Name);
+							JsonObject->SetObjectField(TEXT("name"), NameJson);
+						}
+						else
+						{
+							JsonObject->SetStringField(TEXT("name"), InCodeGenObject->Name.ToString(true));
+						}
+
+						TArray<TSharedPtr<FJsonValue>> PropertiesJson;
+						ToJson(InCodeGenObject->Properties, PropertiesJson);
+						JsonObject->SetArrayField(TEXT("properties"), PropertiesJson);
+
+						return JsonObject;
+					}
+
+					/** Class ToJson */
+					template <>
+					TSharedPtr<FJsonObject> ToJson(const TSharedPtr<FWebAPICodeGenClass>& InCodeGenObject)
+					{
+						check(InCodeGenObject);
+
+						const TSharedPtr<FJsonObject> JsonObject = ToJson<FWebAPICodeGenStruct>(StaticCastSharedPtr<FWebAPICodeGenStruct>(InCodeGenObject));
+
+						TArray<TSharedPtr<FJsonValue>> FunctionsJson;
+						ToJson(InCodeGenObject->Functions, FunctionsJson);
+						JsonObject->SetArrayField(TEXT("functions"), FunctionsJson);
+
+						return JsonObject;
+					}
+
+					/** FunctionParameter ToJson */
+					template <>
+					TSharedPtr<FJsonObject> ToJson(const TSharedPtr<FWebAPICodeGenFunctionParameter>& InCodeGenObject)
+					{
+						check(InCodeGenObject);
+
+						const TSharedPtr<FJsonObject> JsonObject = ToJson<FWebAPICodeGenProperty>(StaticCastSharedPtr<FWebAPICodeGenProperty>(InCodeGenObject));
+						
+						return JsonObject;
+					}
+
+					/** Function ToJson */
+					template <>
+					TSharedPtr<FJsonObject> ToJson(const TSharedPtr<FWebAPICodeGenFunction>& InCodeGenObject)
+					{
+						check(InCodeGenObject);
+
+						const TSharedPtr<FJsonObject> JsonObject = ToJson<FWebAPICodeGenBase>(StaticCastSharedPtr<FWebAPICodeGenBase>(InCodeGenObject));
+
+						if(InCodeGenObject->Name.HasNameInfo())
+						{
+							const TSharedPtr<FJsonObject> NameJson = ToJson(InCodeGenObject->Name);
+							JsonObject->SetObjectField(TEXT("name"), NameJson);
+						}
+						else
+						{
+							JsonObject->SetStringField(TEXT("name"), InCodeGenObject->Name.ToString(true));
+						}
+
+						JsonObject->SetBoolField(TEXT("isOverride"), InCodeGenObject->bIsOverride);
+						JsonObject->SetBoolField(TEXT("isConst"), InCodeGenObject->bIsOverride);
+
+						if(InCodeGenObject->ReturnType.HasTypeInfo())
+						{
+							const TSharedPtr<FJsonObject> TypeJson = ToJson(InCodeGenObject->ReturnType);
+							JsonObject->SetObjectField(TEXT("returnType"), TypeJson);
+						}
+						else
+						{
+							JsonObject->SetStringField(TEXT("returnType"), InCodeGenObject->ReturnType.ToString(true));
+						}
+						
+						JsonObject->SetBoolField(TEXT("isReturnTypeConst"), InCodeGenObject->bIsConstReturnType);
+						JsonObject->SetStringField(TEXT("body"), InCodeGenObject->Body);
+
+						TArray<TSharedPtr<FJsonValue>> ParameterJson;
+						ToJson(InCodeGenObject->Parameters, ParameterJson);
+						JsonObject->SetArrayField(TEXT("parameters"), ParameterJson);
+
+						return JsonObject;
+					}
+
+					/** OperationParameter ToJson */
+					template <>
+					TSharedPtr<FJsonObject> ToJson(const TSharedPtr<FWebAPICodeGenOperationParameter>& InCodeGenObject)
+					{
+						check(InCodeGenObject);
+
+						const TSharedPtr<FJsonObject> JsonObject = ToJson<FWebAPICodeGenProperty>(StaticCastSharedPtr<FWebAPICodeGenProperty>(InCodeGenObject));
+
+						JsonObject->SetStringField(TEXT("storage"), UE::WebAPI::WebAPIParameterStorage::ToString(InCodeGenObject->Storage));
+						JsonObject->SetStringField(TEXT("mediaType"), InCodeGenObject->MediaType);
+
+						return JsonObject;
+					}
+
+					/** OperationRequest ToJson */
+					template <>
+					TSharedPtr<FJsonObject> ToJson(const TSharedPtr<FWebAPICodeGenOperationRequest>& InCodeGenObject)
+					{
+						check(InCodeGenObject);
+
+						const TSharedPtr<FJsonObject> JsonObject = ToJson<FWebAPICodeGenBase>(StaticCastSharedPtr<FWebAPICodeGenBase>(InCodeGenObject));
+
+						if(InCodeGenObject->Name.HasTypeInfo())
+						{
+							const TSharedPtr<FJsonObject> NameJson = ToJson(InCodeGenObject->Name);
+							JsonObject->SetObjectField(TEXT("name"), NameJson);
+						}
+						else
+						{
+							JsonObject->SetStringField(TEXT("name"), InCodeGenObject->Name.ToString(true));
+						}
+
+						TArray<TSharedPtr<FJsonValue>> ParametersJson;
+						ToJson(InCodeGenObject->Parameters, ParametersJson);
+						JsonObject->SetArrayField(TEXT("properties"), ParametersJson); // encode as properties
+
+						return JsonObject;
+					}
+
+					/** OperationResponse ToJson */
+					template <>
+					TSharedPtr<FJsonObject> ToJson(const TSharedPtr<FWebAPICodeGenOperationResponse>& InCodeGenObject)
+					{
+						check(InCodeGenObject);
+
+						const TSharedPtr<FJsonObject> JsonObject = ToJson<FWebAPICodeGenStruct>(StaticCastSharedPtr<FWebAPICodeGenStruct>(InCodeGenObject));
+
+						if(InCodeGenObject->Name.HasTypeInfo())
+						{
+							const TSharedPtr<FJsonObject> NameJson = ToJson(InCodeGenObject->Name);
+							JsonObject->SetObjectField(TEXT("name"), NameJson);
+						}
+						else
+						{
+							JsonObject->SetStringField(TEXT("name"), InCodeGenObject->Name.ToString(true));
+						}
+
+						JsonObject->SetNumberField(TEXT("responseCode"), InCodeGenObject->ResponseCode);
+						JsonObject->SetStringField(TEXT("message"), InCodeGenObject->Message);
+
+						const TSharedPtr<FJsonObject> BaseJson = ToJson(InCodeGenObject->Base);
+						JsonObject->SetObjectField(TEXT("base"), BaseJson);
+
+						return JsonObject;
+					}
+
+					/** Operation ToJson */
+					template <>
+					TSharedPtr<FJsonObject> ToJson(const TSharedPtr<FWebAPICodeGenOperation>& InCodeGenObject)
+					{
+						check(InCodeGenObject);
+
+						const TSharedPtr<FJsonObject> JsonObject = ToJson<FWebAPICodeGenBase>(StaticCastSharedPtr<FWebAPICodeGenBase>(InCodeGenObject));
+
+						if(InCodeGenObject->Name.HasTypeInfo())
+						{
+							const TSharedPtr<FJsonObject> NameJson = ToJson(InCodeGenObject->Name);
+							JsonObject->SetObjectField(TEXT("name"), NameJson);
+						}
+						else
+						{
+							JsonObject->SetStringField(TEXT("name"), InCodeGenObject->Name.ToString(true));
+						}
+
+						JsonObject->SetStringField(TEXT("path"), InCodeGenObject->Path);
+						JsonObject->SetStringField(TEXT("service"), InCodeGenObject->ServiceName);
+						JsonObject->SetStringField(TEXT("verb"), InCodeGenObject->Verb);
+
+						TArray<TSharedPtr<FJsonValue>> RequestsJson;
+						ToJson(InCodeGenObject->Requests, RequestsJson);
+						JsonObject->SetArrayField(TEXT("requests"), RequestsJson);
+
+						TArray<TSharedPtr<FJsonValue>> ResponsesJson;
+						ToJson(InCodeGenObject->Responses, ResponsesJson);
+						JsonObject->SetArrayField(TEXT("responses"), ResponsesJson);
+
+						JsonObject->SetStringField(TEXT("settingsTypeName"), InCodeGenObject->SettingsTypeName);
+
+						return JsonObject;
+					 }
+
+					/** Settings ToJson */
+					template <>
+					TSharedPtr<FJsonObject> ToJson(const TSharedPtr<FWebAPICodeGenSettings>& InCodeGenObject)
+					{
+						check(InCodeGenObject);
+
+						const TSharedPtr<FJsonObject> JsonObject = ToJson<FWebAPICodeGenBase>(StaticCastSharedPtr<FWebAPICodeGenBase>(InCodeGenObject));
+
+						if(InCodeGenObject->Name.HasTypeInfo())
+						{
+							const TSharedPtr<FJsonObject> NameJson = ToJson(InCodeGenObject->Name);
+							JsonObject->SetObjectField(TEXT("name"), NameJson);
+						}
+						else
+						{
+							JsonObject->SetStringField(TEXT("name"), InCodeGenObject->Name.ToString(true));
+						}
+
+						if(InCodeGenObject->ParentType.IsValid())
+						{
+							JsonObject->SetStringField(TEXT("base"), InCodeGenObject->ParentType->GetName());	
+						}
+						JsonObject->SetStringField(TEXT("host"), InCodeGenObject->Host);
+						JsonObject->SetStringField(TEXT("baseUrl"), InCodeGenObject->BaseUrl);
+						JsonObject->SetStringField(TEXT("userAgent"), InCodeGenObject->UserAgent);
+						JsonObject->SetStringField(TEXT("dateTimeFormat"), InCodeGenObject->DateTimeFormat);
+
+						TArray<TSharedPtr<FJsonValue>> SchemesJson;
+						ToJson(InCodeGenObject->Schemes, SchemesJson);
+						JsonObject->SetArrayField(TEXT("schemes"), SchemesJson);
+
+						return JsonObject;
+					}
+
+					/** File ToJson */
+					template <>
+					TSharedPtr<FJsonObject> ToJson(const TSharedPtr<FWebAPICodeGenFile>& InCodeGenObject)
+					{
+						check(InCodeGenObject);
+
+						const TFunction<FString(FWebAPITypeNameVariant)> MakeForwardDeclarationString =
+							[](const FWebAPITypeNameVariant& InTypeName)
+							{
+								const FString ObjectType = InTypeName.TypeInfo->IsEnum()
+															? TEXT("enum class")
+															: InTypeName.TypeInfo->Prefix == TEXT("F")
+																? TEXT("struct")
+																: TEXT("class");
+
+								const FString Suffix = InTypeName.TypeInfo->IsEnum() ? TEXT(" : uint8") : TEXT("");								
+								
+								return FString::Printf(TEXT("%s %s%s"), *ObjectType, *InTypeName.ToString(), *Suffix);
+							};
+
+						const TSharedPtr<FJsonObject> JsonObject = MakeShared<FJsonObject>();
+
+						JsonObject->SetStringField(TEXT("baseFilePath"), InCodeGenObject->BaseFilePath);
+						JsonObject->SetStringField(TEXT("relativeFilePath"), InCodeGenObject->RelativeFilePath);
+						JsonObject->SetStringField(TEXT("fileName"), InCodeGenObject->FileName);
+						JsonObject->SetStringField(TEXT("fileType"), InCodeGenObject->FileType);
+						JsonObject->SetStringField(TEXT("namespace"), InCodeGenObject->Namespace);
+						JsonObject->SetStringField(TEXT("module"), InCodeGenObject->Module);
+						JsonObject->SetStringField(TEXT("copyrightNotice"), InCodeGenObject->CopyrightNotice);
+
+						TArray<TSharedPtr<FJsonValue>> IncludePathJson;
+						ToJson(InCodeGenObject->IncludePaths, IncludePathJson);
+						JsonObject->SetArrayField(TEXT("includePaths"), IncludePathJson);
+
+						TArray<FString> ForwardDeclarations;
+
+						// Find and transform Enums
+						TArray<TSharedPtr<FJsonValue>> EnumsJson;
+						Algo::TransformIf(
+							InCodeGenObject->SubItems,
+							EnumsJson,
+							[](const TSharedPtr<FWebAPICodeGenBase>& InItem)
+							{
+								return InItem->GetTypeName() == TEXT("Enum");
+							},
+							[&](const TSharedPtr<FWebAPICodeGenBase>& InItem)
+							{
+								const TSharedPtr<FWebAPICodeGenEnum> Enum = StaticCastSharedPtr<FWebAPICodeGenEnum>(InItem);
+								ForwardDeclarations.Add(MakeForwardDeclarationString(Enum->Name));
+								
+								TSharedPtr<FJsonObject> EnumJson = UE::WebAPI::Generator::LiquidJS::Private::ToJson(Enum);
+								return MakeShared<FJsonValueObject>(EnumJson);
+							});
+
+						// Find and transform Structs
+						TArray<TSharedPtr<FJsonValue>> StructsJson;
+						Algo::TransformIf(
+							InCodeGenObject->SubItems,
+							StructsJson,
+							[](const TSharedPtr<FWebAPICodeGenBase>& InItem)
+							{
+								return InItem->GetTypeName() == TEXT("Struct");
+							},
+							[&](const TSharedPtr<FWebAPICodeGenBase>& InItem)
+							{
+								const TSharedPtr<FWebAPICodeGenStruct> Struct = StaticCastSharedPtr<FWebAPICodeGenStruct>(InItem);
+								ForwardDeclarations.Add(MakeForwardDeclarationString(Struct->Name));
+
+								TSharedPtr<FJsonObject> StructJson = UE::WebAPI::Generator::LiquidJS::Private::ToJson(Struct);								
+								return MakeShared<FJsonValueObject>(StructJson);
+							});
+
+						// Find and transform Classes
+						TArray<TSharedPtr<FJsonValue>> ClassesJson;
+						Algo::TransformIf(
+							InCodeGenObject->SubItems,
+							ClassesJson,
+							[](const TSharedPtr<FWebAPICodeGenBase>& InItem)
+							{
+								return InItem->GetTypeName() == TEXT("Class");
+							},
+							[&](const TSharedPtr<FWebAPICodeGenBase>& InItem)
+							{
+								const TSharedPtr<FWebAPICodeGenClass> Class = StaticCastSharedPtr<FWebAPICodeGenClass>(InItem);
+								ForwardDeclarations.Add(MakeForwardDeclarationString(Class->Name));
+								
+								TSharedPtr<FJsonObject> ClassJson = UE::WebAPI::Generator::LiquidJS::Private::ToJson(Class);
+								return MakeShared<FJsonValueObject>(ClassJson);
+							});
+
+						// Find and transform Operations
+						TArray<TSharedPtr<FJsonValue>> OperationsJson;
+						TArray<TSharedPtr<FJsonValue>> RequestsJson;
+						TArray<TSharedPtr<FJsonValue>> ResponsesJson;
+						Algo::TransformIf(
+							InCodeGenObject->SubItems,
+							OperationsJson,
+							[](const TSharedPtr<FWebAPICodeGenBase>& InItem)
+							{
+								return InItem->GetTypeName() == TEXT("Operation");
+							},
+							[&](const TSharedPtr<FWebAPICodeGenBase>& InItem)
+							{
+								const TSharedPtr<FWebAPICodeGenOperation> Operation = StaticCastSharedPtr<FWebAPICodeGenOperation>(InItem);
+								TSharedPtr<FJsonObject> OperationJson = UE::WebAPI::Generator::LiquidJS::Private::ToJson(Operation);
+
+								ForwardDeclarations.Add(MakeForwardDeclarationString(Operation->Name));
+
+								RequestsJson.Append(OperationJson->GetArrayField(TEXT("requests")));
+								ResponsesJson.Append(OperationJson->GetArrayField(TEXT("responses")));
+								
+								return MakeShared<FJsonValueObject>(OperationJson);
+							});
+
+						// Find and transform Settings
+						TArray<TSharedPtr<FJsonObject>> SettingsJsonArray;
+						Algo::TransformIf(
+							InCodeGenObject->SubItems,
+							SettingsJsonArray,
+							[](const TSharedPtr<FWebAPICodeGenBase>& InItem)
+							{
+								return InItem->GetTypeName() == TEXT("Settings");
+							},
+							[](const TSharedPtr<FWebAPICodeGenBase>& InItem)
+							{
+								const TSharedPtr<FWebAPICodeGenSettings> Settings = StaticCastSharedPtr<FWebAPICodeGenSettings>(InItem);
+								TSharedPtr<FJsonObject> SettingsJson = UE::WebAPI::Generator::LiquidJS::Private::ToJson(Settings);
+								return SettingsJson;
+							});
+
+						// Set here as Operations can add additional structs/enums
+						JsonObject->SetArrayField(TEXT("enums"), EnumsJson);
+						JsonObject->SetArrayField(TEXT("structs"), StructsJson);
+						JsonObject->SetArrayField(TEXT("classes"), ClassesJson);
+						JsonObject->SetArrayField(TEXT("operations"), OperationsJson);
+						JsonObject->SetArrayField(TEXT("requests"), RequestsJson);
+						JsonObject->SetArrayField(TEXT("responses"), ResponsesJson);
+
+						TArray<TSharedPtr<FJsonValue>> ForwardDeclarationsJson;
+						ToJson(ForwardDeclarations, ForwardDeclarationsJson);
+						JsonObject->SetArrayField(TEXT("forwardDeclarations"), ForwardDeclarationsJson);
+
+						if(!SettingsJsonArray.IsEmpty())
+						{
+							JsonObject->SetObjectField(TEXT("settings"), SettingsJsonArray.Last());							
+						}
+
+						return JsonObject;
+					}
+
+					template <>
+					void ToJson<FWebAPICodeGenBase>(const TArray<TSharedPtr<FWebAPICodeGenBase>>& InCodeGenObjects, TArray<TSharedPtr<FJsonValue>>& OutJson)
+					{
+						OutJson.Reserve(InCodeGenObjects.Num());
+						for(const TSharedPtr<FWebAPICodeGenBase>& CodeGenObject : InCodeGenObjects)
+						{
+							TSharedPtr<FJsonObject> ElementJson = MakeShared<FJsonObject>();
+							const FName& CodeGenObjectType = CodeGenObject->GetTypeName();
+							if(CodeGenObjectType == TEXT("Enum"))
+							{
+								ElementJson = ToJson(StaticCastSharedPtr<FWebAPICodeGenEnum>(CodeGenObject));
+							}
+							else if(CodeGenObjectType == TEXT("Struct"))
+							{
+								ElementJson = ToJson(StaticCastSharedPtr<FWebAPICodeGenStruct>(CodeGenObject));
+							}
+							else if(CodeGenObjectType == TEXT("Operation"))
+							{
+								ElementJson = ToJson(StaticCastSharedPtr<FWebAPICodeGenOperation>(CodeGenObject));
+							}
+							else if(CodeGenObjectType == TEXT("Class"))
+							{
+								ElementJson = ToJson(StaticCastSharedPtr<FWebAPICodeGenClass>(CodeGenObject));
+							}
+							OutJson.Add(MakeShared<FJsonValueObject>(ElementJson));
+						}
+					}
+
+					/** Runs the given function on the GameThread, returning a Future. */
+					template<typename ResultType>
+					static TFuture<ResultType> ExecuteOnGameThread(const TFunction<ResultType()>& Function)
+					{
+						TSharedRef<TPromise<ResultType>, ESPMode::ThreadSafe> Promise = MakeShared<TPromise<ResultType>>();
+						TFunction<void()> PromiseKeeper =
+							[Promise, Function]
+						{
+								Promise->SetValue(Function());
+							};
+						
+						if (!IsInGameThread())
+						{
+							AsyncTask(ENamedThreads::GameThread, MoveTemp(PromiseKeeper));
+						}
+						else
+						{
+							PromiseKeeper();
+						}
+						return Promise->GetFuture();
+					}
+				}
+			}
+		}
+	}
+}
+
+TFuture<bool> UWebAPILiquidJSCodeGenerator::IsAvailable()
+{
+	const TSharedRef<TPromise<bool>, ESPMode::ThreadSafe> Promise = MakeShared<TPromise<bool>, ESPMode::ThreadSafe>();
+
+	FHttpModule& HttpModule = FHttpModule::Get();
+
+	const FString URL = GetMutableDefault<UWebAPILiquidJSSettings>()->GetServiceUrl(TEXT("/api/ping"));
+
+	const TSharedRef<IHttpRequest, ESPMode::ThreadSafe> Request = HttpModule.CreateRequest();
+	Request->SetVerb(TEXT("GET"));
+	Request->SetURL(URL);
+	Request->SetHeader(TEXT("User-Agent"), TEXT("X-UnrealEngine-Agent"));
+	Request->SetHeader(TEXT("Accept"), TEXT("application/json"));
+	Request->SetHeader(TEXT("Accept-Encoding"), TEXT("identity"));
+	Request->SetTimeout(10.0f); // 10 second timeout
+
+	std::atomic<bool> bHasRetried = false;
+
+	Request->OnProcessRequestComplete().BindLambda(
+	[Promise, &bHasRetried, Request](FHttpRequestPtr, FHttpResponsePtr InResponse, bool bInWasSuccessful)
+	{
+		if(bInWasSuccessful && EHttpResponseCodes::IsOk(InResponse->GetResponseCode()))
+		{
+			Promise->SetValue(true);
+		}
+		else
+		{
+			if(!bHasRetried.load(std::memory_order_relaxed))
+			{
+				// @todo: what's this? can I use it for retry system?
+				// Request->OnRequestWillRetry()
+
+				// Try starting
+				bool bRetryResult = FModuleManager::Get().GetModulePtr<FWebAPILiquidJSModule>(TEXT("WebAPILiquidJS"))->TryStartWebApp();
+				bHasRetried.store(true, std::memory_order_relaxed);
+				Request->ProcessRequest();
+				return;
+			}
+
+			Promise->SetValue(false);
+		}
+	});
+
+	Request->ProcessRequest();
+
+	return Promise->GetFuture();
+}
+
+TFuture<EWebAPIGenerationResult> UWebAPILiquidJSCodeGenerator::GenerateFile(const TWeakObjectPtr<UWebAPIDefinition>& InDefinition, const TSharedPtr<FWebAPICodeGenFile>& InFile)
+{
+	check(InDefinition.IsValid());
+	
+	const TSharedRef<TPromise<EWebAPIGenerationResult>, ESPMode::ThreadSafe> Promise = MakeShared<TPromise<EWebAPIGenerationResult>, ESPMode::ThreadSafe>();
+
+	FHttpModule& HttpModule = FHttpModule::Get();
+
+	const FString FileType = InFile->FileType.EndsWith(TEXT("h"))
+							? TEXT("decl")
+							: TEXT("defn");
+
+	const FString URL = GetDefault<UWebAPILiquidJSSettings>()->GetServiceUrl(TEXT("api/gen/") + FileType);
+
+	const TSharedRef<IHttpRequest, ESPMode::ThreadSafe> Request = HttpModule.CreateRequest();
+	Request->SetVerb(TEXT("POST"));
+	Request->SetURL(URL);
+	Request->SetHeader(TEXT("User-Agent"), TEXT("X-UnrealEngine-Agent"));
+	Request->SetHeader(TEXT("Content-Type"), TEXT("application/json"));
+	Request->SetHeader(TEXT("Accept"), TEXT("application/json"));
+	Request->SetHeader(TEXT("Accept-Encoding"), TEXT("identity"));
+
+	const TSharedPtr<FJsonObject> FileJson = UE::WebAPI::Generator::LiquidJS::Private::ToJson(InFile);
+#if WITH_WEBAPI_DEBUG
+	bool bWriteResult = InDefinition->GetGeneratorSettings().bWriteResult;
+#else
+	bool bWriteResult = true;
+#endif
+
+	FString JsonString;
+	const TSharedRef<TJsonWriter<TCHAR, TCondensedJsonPrintPolicy<TCHAR>>> Writer = TJsonWriterFactory<TCHAR, TCondensedJsonPrintPolicy<TCHAR>>::Create(&JsonString);
+	FJsonSerializer::Serialize(FileJson.ToSharedRef(), Writer);
+
+	Request->SetContentAsString(JsonString);
+
+	Request->OnProcessRequestComplete().BindLambda(
+	[Promise, InFile, bWriteResult, InDefinition](FHttpRequestPtr, FHttpResponsePtr InResponse, bool bInWasSuccessful)
+	{
+		if(bInWasSuccessful && InResponse->GetResponseCode() == 200)
+		{
+			const FString GeneratedCodeString = InResponse->GetContentAsString();
+
+			// Write the result to the original schema objects (allows for visualization and debug of generated code)
+ 			Algo::ForEach(InFile->SchemaObjects, [&GeneratedCodeString](const TWeakObjectPtr<UObject>& InSchemaObject)
+			{
+				if(InSchemaObject.IsValid())
+				{
+					const TScriptInterface<IWebAPISchemaObjectInterface> AsInterface = InSchemaObject.Get();
+					AsInterface.GetInterface()->AppendCodeText(GeneratedCodeString);
+				}
+			});
+			
+			if(bWriteResult)
+			{
+				FFileHelper::SaveStringToFile(GeneratedCodeString, *InFile->GetFullPath());	
+			}
+			
+			Promise->SetValue(EWebAPIGenerationResult::Succeeded);
+		}
+		else
+		{
+			const FString JsonString = InResponse->GetContentAsString();				
+			const TSharedRef<TJsonReader<TCHAR>> Reader = TJsonReaderFactory<TCHAR>::Create(JsonString);
+
+			FFormatNamedArguments Args;
+			Args.Add(TEXT("Code"), InResponse->GetResponseCode());
+			Args.Add(TEXT("Url"), FText::FromString(InResponse->GetURL()));
+			
+			TSharedPtr<FJsonObject> JsonObject;
+			if (FJsonSerializer::Deserialize(Reader, JsonObject) && JsonObject.IsValid())
+			{
+				FString ErrorName;
+				JsonObject->TryGetStringField(TEXT("errorName"), ErrorName);
+
+				FString Message;
+				JsonObject->TryGetStringField(TEXT("message"), Message);
+
+				FString Stack;
+				JsonObject->TryGetStringField(TEXT("stack"), Stack);
+
+				Args.Add(TEXT("Error"), FText::FromString(ErrorName));
+				Args.Add(TEXT("Message"), FText::FromString(Message));
+				Args.Add(TEXT("StackTrace"), FText::FromString(Stack));
+				InDefinition->GetMessageLog()->LogError(FText::Format(LOCTEXT("HttpErrorWithMessage", "HTTP Response Code: {Code}\nError: {Error}\nMessage: {Message}\nStackTrace: {StackTrace}"), Args), UWebAPILiquidJSCodeGenerator::LogName);
+			}
+			else
+			{
+				InDefinition->GetMessageLog()->LogError(FText::Format(LOCTEXT("HttpError", "HTTP Response Code: {Code}\nUrl: {Url}"), Args), UWebAPILiquidJSCodeGenerator::LogName);
+			}
+
+			Promise->SetValue(EWebAPIGenerationResult::Failed);
+		}
+	});
+
+	Request->ProcessRequest();
+
+	return Promise->GetFuture();
+}
+
+#undef LOCTEXT_NAMESPACE
diff --git a/Engine/Plugins/Web/WebAPI/Source/WebAPILiquidJS/Private/WebAPILiquidJSCodeGenerator.h b/Engine/Plugins/Web/WebAPI/Source/WebAPILiquidJS/Private/WebAPILiquidJSCodeGenerator.h
new file mode 100644
--- /dev/null
+++ b/Engine/Plugins/Web/WebAPI/Source/WebAPILiquidJS/Private/WebAPILiquidJSCodeGenerator.h
@@ -0,0 +1,22 @@
+﻿// Copyright Epic Games, Inc. All Rights Reserved.
+
+#pragma once
+
+#include "CoreMinimal.h"
+#include "CodeGen/WebAPICodeGenerator.h"
+
+#include "WebAPILiquidJSCodeGenerator.generated.h"
+
+/**
+ * 
+ */
+UCLASS()
+class WEBAPILIQUIDJS_API UWebAPILiquidJSCodeGenerator
+	: public UWebAPICodeGeneratorBase
+{
+	GENERATED_BODY()
+	
+public:
+	virtual TFuture<bool> IsAvailable() override;
+	virtual TFuture<EWebAPIGenerationResult> GenerateFile(const TWeakObjectPtr<UWebAPIDefinition>& InDefinition, const TSharedPtr<FWebAPICodeGenFile>& InFile) override;
+};
diff --git a/Engine/Plugins/Web/WebAPI/Source/WebAPILiquidJS/Private/WebAPILiquidJSLog.cpp b/Engine/Plugins/Web/WebAPI/Source/WebAPILiquidJS/Private/WebAPILiquidJSLog.cpp
new file mode 100644
--- /dev/null
+++ b/Engine/Plugins/Web/WebAPI/Source/WebAPILiquidJS/Private/WebAPILiquidJSLog.cpp
@@ -0,0 +1,5 @@
+// Copyright Epic Games, Inc. All Rights Reserved.
+
+#include "WebAPILiquidJSLog.h"
+
+DEFINE_LOG_CATEGORY(LogWebAPILiquidJS);
diff --git a/Engine/Plugins/Web/WebAPI/Source/WebAPILiquidJS/Private/WebAPILiquidJSLog.h b/Engine/Plugins/Web/WebAPI/Source/WebAPILiquidJS/Private/WebAPILiquidJSLog.h
new file mode 100644
--- /dev/null
+++ b/Engine/Plugins/Web/WebAPI/Source/WebAPILiquidJS/Private/WebAPILiquidJSLog.h
@@ -0,0 +1,7 @@
+// Copyright Epic Games, Inc. All Rights Reserved.
+
+#pragma once
+
+#include "CoreMinimal.h"
+
+DECLARE_LOG_CATEGORY_EXTERN(LogWebAPILiquidJS, Log, All);
diff --git a/Engine/Plugins/Web/WebAPI/Source/WebAPILiquidJS/Private/WebAPILiquidJSModule.cpp b/Engine/Plugins/Web/WebAPI/Source/WebAPILiquidJS/Private/WebAPILiquidJSModule.cpp
new file mode 100644
--- /dev/null
+++ b/Engine/Plugins/Web/WebAPI/Source/WebAPILiquidJS/Private/WebAPILiquidJSModule.cpp
@@ -0,0 +1,55 @@
+// Copyright Epic Games, Inc. All Rights Reserved.
+
+#include "WebAPILiquidJSModule.h"
+
+#include "IWebAPIEditorModule.h"
+#include "WebAPILiquidJSCodeGenerator.h"
+#include "WebAPILiquidJSProcess.h"
+#include "Misc/CoreDelegates.h"
+
+#define LOCTEXT_NAMESPACE "WebAPILiquidJS"
+
+void FWebAPILiquidJSModule::StartupModule()
+{
+	IWebAPIEditorModuleInterface& WebAPIEditorModule = IWebAPIEditorModuleInterface::Get();
+
+	WebAPIEditorModule.AddCodeGenerator(TEXT("LiquidJS"), UWebAPILiquidJSCodeGenerator::StaticClass());
+
+	WebApp = MakeShared<FWebAPILiquidJSProcess>();
+	TSharedPtr<FWebAPILiquidJSProcess> WebAppLocal = WebApp;
+	FCoreDelegates::OnFEngineLoopInitComplete.AddLambda([WebAppLocal]
+	{
+		WebAppLocal->TryStart();
+	});
+}
+
+void FWebAPILiquidJSModule::ShutdownModule()
+{
+	WebApp->Shutdown();
+	FCoreDelegates::OnFEngineLoopInitComplete.RemoveAll(this);
+	
+	static FName WebAPIEditorModuleName = TEXT("WebAPIEditorModule");
+	if(FModuleManager::Get().IsModuleLoaded(WebAPIEditorModuleName))
+	{
+		IWebAPIEditorModuleInterface& WebAPIEditorModule = IWebAPIEditorModuleInterface::Get();
+		WebAPIEditorModule.RemoveCodeGenerator(TEXT("LiquidJS"));
+	}
+}
+
+bool FWebAPILiquidJSModule::TryStartWebApp() const
+{
+	if(!WebApp.IsValid())
+	{
+		return false;
+	}
+	else if(WebApp->GetStatus() >= FWebAPILiquidJSProcess::EStatus::Launching)
+	{
+		return true;
+	}
+
+	return WebApp->TryStart();
+}
+
+#undef LOCTEXT_NAMESPACE
+
+IMPLEMENT_MODULE(FWebAPILiquidJSModule, WebAPILiquidJS)
diff --git a/Engine/Plugins/Web/WebAPI/Source/WebAPILiquidJS/Private/WebAPILiquidJSModule.h b/Engine/Plugins/Web/WebAPI/Source/WebAPILiquidJS/Private/WebAPILiquidJSModule.h
new file mode 100644
--- /dev/null
+++ b/Engine/Plugins/Web/WebAPI/Source/WebAPILiquidJS/Private/WebAPILiquidJSModule.h
@@ -0,0 +1,21 @@
+// Copyright Epic Games, Inc. All Rights Reserved.
+
+#pragma once
+
+#include "IWebAPILiquidJSModule.h"
+
+#include "CoreMinimal.h"
+
+class FWebAPILiquidJSModule final
+    : public IWebAPILiquidJSModuleInterface
+{
+public:
+	virtual void StartupModule() override;
+	virtual void ShutdownModule() override;
+
+	bool TryStartWebApp() const;
+
+private:
+	/** The actual process that runs the middleman server. */
+	TSharedPtr<class FWebAPILiquidJSProcess> WebApp;
+};
diff --git a/Engine/Plugins/Web/WebAPI/Source/WebAPILiquidJS/Private/WebAPILiquidJSProcess.cpp b/Engine/Plugins/Web/WebAPI/Source/WebAPILiquidJS/Private/WebAPILiquidJSProcess.cpp
new file mode 100644
--- /dev/null
+++ b/Engine/Plugins/Web/WebAPI/Source/WebAPILiquidJS/Private/WebAPILiquidJSProcess.cpp
@@ -0,0 +1,280 @@
+﻿// Copyright Epic Games, Inc. All Rights Reserved.
+
+#include "WebAPILiquidJSProcess.h"
+
+#include "IWebAPILiquidJSModule.h"
+#include "Interfaces/IPluginManager.h"
+#include "GenericPlatform/GenericPlatformFile.h"
+#include "GenericPlatform/GenericPlatformProcess.h"
+#include "HAL/PlatformFileManager.h"
+#include "Modules/ModuleManager.h"
+#include "SocketSubsystem.h"
+#include "WebAPILiquidJSLog.h"
+#include "WebAPILiquidJSSettings.h"
+#include "INetworkingWebSocket.h"
+
+#define LOCTEXT_NAMESPACE "WebAPILiquidJS"
+
+FWebAPILiquidJSProcess::FWebAPILiquidJSProcess()
+	: Status(EStatus::Stopped)
+{
+}
+
+FWebAPILiquidJSProcess::~FWebAPILiquidJSProcess()
+{
+}
+
+bool FWebAPILiquidJSProcess::TryStart()
+{
+	Shutdown();
+
+	Status = EStatus::Launching;
+
+	IPlatformFile& PlatformFile = FPlatformFileManager::Get().GetPlatformFile();
+
+	const FString WebApp = IPluginManager::Get().FindPlugin("WebAPI")->GetBaseDir() / TEXT("WebAPIGeneratorApp");
+	Root = PlatformFile.ConvertToAbsolutePathForExternalAppForRead(*WebApp);
+
+	if (!PlatformFile.DirectoryExists(*Root))
+	{
+		UE_LOG(LogWebAPILiquidJS, Warning, TEXT("WebApp folder doesn't exist (%s)"), *Root);
+		return false;
+	}
+
+	FAsyncTaskNotificationConfig NotificationConfig;
+	NotificationConfig.bKeepOpenOnFailure = true;
+	NotificationConfig.TitleText = LOCTEXT("WebAPILiquidJS_Launch", "Launching WebAPILiquidJS");
+#if !NO_LOGGING
+	NotificationConfig.LogCategory = &LogWebAPILiquidJS;
+#endif
+	NotificationConfig.bIsHeadless = false;
+	TaskNotification = MakeUnique<FAsyncTaskNotification>(NotificationConfig);
+
+	Thread = FRunnableThread::Create(this, TEXT("FWebAPILiquidJSProcess"), 8 * 1024, TPri_BelowNormal);
+	return true;
+}
+
+void FWebAPILiquidJSProcess::Shutdown()
+{
+	Status = EStatus::Stopped;
+
+	SetExternalLoggerEnabled(false);
+
+	if (Process.IsValid())
+	{
+		FPlatformProcess::TerminateProc(Process, true);
+		Process.Reset();
+	}
+
+	if (Thread)
+	{
+		Thread->Kill(true);
+		delete Thread;
+		Thread = nullptr;
+	}
+}
+
+FWebAPILiquidJSProcess::EStatus FWebAPILiquidJSProcess::GetStatus() const
+{
+	return Status.load();
+}
+
+void FWebAPILiquidJSProcess::SetExternalLoggerEnabled(bool bEnableExternalLog) const
+{
+	TSharedPtr<INetworkingWebSocket> WebSocketLoggerConnection;
+	if (bEnableExternalLog)
+	{
+		const UWebAPILiquidJSSettings* Settings = GetDefault<UWebAPILiquidJSSettings>();
+
+		const TSharedRef<FInternetAddr> Address = ISocketSubsystem::Get()->CreateInternetAddr();
+
+		bool bIsValidIp = false;
+		Address->SetIp(TEXT("127.0.0.1"), bIsValidIp);
+		Address->SetPort(Settings->Port + 2);
+		//WebSocketLoggerConnection = FModuleManager::LoadModuleChecked<IWebSocketNetworkingModule>(TEXT("WebSocketNetworking")).CreateConnection(*Address);
+	}
+
+	IWebAPILiquidJSModuleInterface* WebAPILiquidJSModule = FModuleManager::GetModulePtr<IWebAPILiquidJSModuleInterface>("WebAPILiquidJS");
+	if (WebAPILiquidJSModule)
+	{
+		// @todo
+		//WebAPILiquidJSModule->SetExternalRemoteWebSocketLoggerConnection(WebSocketLoggerConnection);
+	}
+}
+
+uint32 FWebAPILiquidJSProcess::Run()
+{
+	IPlatformFile& PlatformFile = FPlatformFileManager::Get().GetPlatformFile();
+
+#if PLATFORM_WINDOWS
+	const FString StartScript = Root / TEXT("Start.bat");
+#elif PLATFORM_MAC
+	const FString StartScript = Root / TEXT("Start.command");
+#else
+	const FString StartScript = Root / TEXT("Start.sh");
+#endif
+
+	const FText ErrorTitle = LOCTEXT("WebAPILiquidJS_ErrorTitle", "Failed to Launch the WebAPILiquidJS process");
+
+	if (!PlatformFile.FileExists(*StartScript))
+	{
+		TaskNotification->SetComplete(
+			ErrorTitle,
+			LOCTEXT("WebAPILiquidJS_FilesMissing", "Missing files in WebAPIGeneratorApp folder"),
+			false
+		);
+		Status = EStatus::Error;
+		return 0;
+	}
+
+	const UWebAPILiquidJSSettings* Settings = GetDefault<UWebAPILiquidJSSettings>();
+
+	FString Args = FString::Printf(TEXT("--port %d --uews %d --uehttp %d --monitor "),
+									Settings->Port,
+									Settings->WebSocketServerPort,
+									Settings->HttpServerPort);
+
+	if (Settings->bForceWebAppBuildAtStartup)
+	{
+		Args.Append(TEXT("--build "));
+	}
+
+	if (Settings->bWebAppLogRequestDuration)
+	{
+		Args.Append(TEXT("--log "));
+	}
+
+	void* ReadPipe = nullptr;
+	void* WritePipe = nullptr;
+
+	if (!FPlatformProcess::CreatePipe(ReadPipe, WritePipe))
+	{
+		TaskNotification->SetComplete(
+			ErrorTitle,
+			LOCTEXT("WebAPILiquidJS_PipeFailed", "Failed to create Pipes for the WebAPILiquidJS process"),
+			false
+		);
+		Status = EStatus::Error;
+		return 0;
+	}
+
+	check(ReadPipe);
+	check(WritePipe);
+
+	Process = FPlatformProcess::CreateProc(
+		*StartScript,	/* Path to start script */
+		*Args,			/* Arguments with port numbers */
+		false,			/* bLaunchDetached */
+		true,			/* bLaunchHidden */
+		true,			/* bLaunchReallyHidden */
+		nullptr,		/* OutProcessID */
+		0,				/* PriorityModifier */
+		*Root,			/* OptionalWorkingDirectory */
+		WritePipe,		/* PipeWriteChild */
+		ReadPipe		/* PipeReadChild */
+	);
+
+	if (!Process.IsValid())
+	{
+		FPlatformProcess::ClosePipe(ReadPipe, WritePipe);
+		TaskNotification->SetComplete(
+			ErrorTitle,
+			LOCTEXT("WebAPILiquidJS_LaunchFailed", "Failed to start WebAPILiquidJS process"),
+			false
+		);
+
+		Status = EStatus::Error;
+		return 0;
+	}
+
+	bool bLoadDone = false;
+	auto LogReadPipe = [&](void* InReadPipe)
+	{
+		const FString ProcessOutput = FPlatformProcess::ReadPipe(InReadPipe);
+		if (ProcessOutput.Len() > 0)
+		{
+			TArray<FString> Lines;
+			ProcessOutput.ParseIntoArray(Lines, TEXT("\n"), false);
+
+			for (const FString& Line : Lines)
+			{
+				if (Line.Len() == 0)
+				{
+					continue;
+				}
+
+				if (!bLoadDone)
+				{
+					if (Line.StartsWith(TEXT("ERROR: ")))
+					{
+						bLoadDone = true;
+						const FString ErrorMessage = Line.Mid(7);
+						TaskNotification->SetComplete(ErrorTitle, FText::FromString(ErrorMessage), false);
+						Status = EStatus::Error;
+					}
+					else if (Line.StartsWith(TEXT("DONE: ")))
+					{
+						bLoadDone = true;
+						const FString CompleteMessage = Line.Mid(6);
+						TaskNotification->SetComplete(
+							LOCTEXT("WebAPILiquidJS_SuccessTitle", "WebAPILiquidJS is running"),
+							FText::FromString(CompleteMessage),
+							true);
+						Status = EStatus::Running;
+
+						if (Settings->bWebAppLogRequestDuration)
+						{
+							SetExternalLoggerEnabled(true);
+						}
+					}
+					else
+					{
+						TaskNotification->SetProgressText(FText::FromString(Line));
+					}
+				}
+				else
+				{
+					UE_LOG(LogWebAPILiquidJS, Log, TEXT("%s"), *Line);
+				}
+			}
+		}
+
+		return ProcessOutput.Len();
+	};
+
+	UE_LOG(LogWebAPILiquidJS, Log, TEXT("WebApp started, initial launch will take longer as it will be building the WebApp"));
+
+	while (FPlatformProcess::IsProcRunning(Process))
+	{
+		const int32 BytesRead = LogReadPipe(ReadPipe);
+		if (!BytesRead)
+		{
+			FPlatformProcess::Sleep(0.2);
+		}
+	}
+
+	// One last ReadPipe log for anything that the while loop didn't catch.
+	LogReadPipe(ReadPipe);
+
+	if (!bLoadDone)
+	{
+		TaskNotification->SetComplete(
+			ErrorTitle,
+			LOCTEXT("WebAPILiquidJS_LaunchFailed", "WebApp exited"),
+			false
+		);
+	}
+	else
+	{
+		UE_LOG(LogWebAPILiquidJS, Log, TEXT("WebApp exited"));
+	}
+
+	Status = EStatus::Stopped;
+
+	FPlatformProcess::ClosePipe(ReadPipe, WritePipe);
+	Process.Reset();
+
+	return 0;
+}
+
+#undef LOCTEXT_NAMESPACE
\ No newline at end of file
diff --git a/Engine/Plugins/Web/WebAPI/Source/WebAPILiquidJS/Private/WebAPILiquidJSProcess.h b/Engine/Plugins/Web/WebAPI/Source/WebAPILiquidJS/Private/WebAPILiquidJSProcess.h
new file mode 100644
--- /dev/null
+++ b/Engine/Plugins/Web/WebAPI/Source/WebAPILiquidJS/Private/WebAPILiquidJSProcess.h
@@ -0,0 +1,59 @@
+﻿// Copyright Epic Games, Inc. All Rights Reserved.
+
+#pragma once
+
+#include <atomic>
+#include "Containers/UnrealString.h"
+#include "HAL/PlatformProcess.h"
+#include "HAL/Runnable.h"
+#include "HAL/RunnableThread.h"
+#include "Misc/AsyncTaskNotification.h"
+
+class FWebAPILiquidJSProcess : private FRunnable
+{
+public:
+	enum class EStatus : uint8
+	{
+		Stopped = 0,
+		Error = 1,
+		
+		Launching = 2,
+		Running = 3,
+	};
+
+public:
+	FWebAPILiquidJSProcess();
+	
+	virtual ~FWebAPILiquidJSProcess();
+
+	/** Starts the webapp external process, first launch will take some time as it needs to compile the webapp */
+	bool TryStart();
+
+	/** Shutdown the webapp external process */
+	void Shutdown();
+
+	/** Get the webapp's status. */
+	EStatus GetStatus() const;
+
+	/** Enable / disable external logger */
+	void SetExternalLoggerEnabled(bool bEnableExternalLog) const;
+
+protected:
+	virtual uint32 Run() override;
+
+private:
+	/** Plugin root folder */
+	FString Root;
+
+	/** WebApp Node.js process handle */
+	FProcHandle Process;
+
+	/** Thread for WebApp node.js process (read stdout and forward to log) */
+	FRunnableThread* Thread = nullptr;
+
+	/** Progress notification */
+	TUniquePtr<FAsyncTaskNotification> TaskNotification;
+
+	/** The current status of the web app. */
+	std::atomic<EStatus> Status;
+};
\ No newline at end of file
diff --git a/Engine/Plugins/Web/WebAPI/Source/WebAPILiquidJS/Private/WebAPILiquidJSSettings.cpp b/Engine/Plugins/Web/WebAPI/Source/WebAPILiquidJS/Private/WebAPILiquidJSSettings.cpp
new file mode 100644
--- /dev/null
+++ b/Engine/Plugins/Web/WebAPI/Source/WebAPILiquidJS/Private/WebAPILiquidJSSettings.cpp
@@ -0,0 +1,38 @@
+﻿// Copyright Epic Games, Inc. All Rights Reserved.
+
+#include "WebAPILiquidJSSettings.h"
+
+FString UWebAPILiquidJSSettings::GetServiceUrl(const FString& InSubPath)
+{
+	// Remove prefix / to avoid double-up
+	FString SubPath = InSubPath;
+	SubPath.RemoveFromStart(TEXT("/"));
+	
+	if(Port != FormattedWithPort)
+	{
+		FormattedServiceUrl = FString::Format(*ServiceUrl, FStringFormatNamedArguments(
+			{{	TEXT("Port"), Port }}
+		));
+		FormattedServiceUrl.RemoveFromEnd(TEXT("/"));
+		FormattedWithPort = Port;
+
+		check(!FormattedServiceUrl.Contains(TEXT("{")));
+		check(!FormattedServiceUrl.Contains(TEXT("}")));
+	}
+
+	return FormattedServiceUrl + TEXT("/") + SubPath;
+}
+
+FString UWebAPILiquidJSSettings::GetServiceUrl(const FString& InSubPath) const
+{
+	if(Port != FormattedWithPort)
+	{
+		// Need to call non-const version before this is ever called!
+		checkNoEntry();
+	}
+
+	FString SubPath = InSubPath;
+	SubPath.RemoveFromStart(TEXT("/"));
+
+	return FormattedServiceUrl + TEXT("/") + SubPath;
+}
diff --git a/Engine/Plugins/Web/WebAPI/Source/WebAPILiquidJS/Private/WebAPILiquidJSSettings.h b/Engine/Plugins/Web/WebAPI/Source/WebAPILiquidJS/Private/WebAPILiquidJSSettings.h
new file mode 100644
--- /dev/null
+++ b/Engine/Plugins/Web/WebAPI/Source/WebAPILiquidJS/Private/WebAPILiquidJSSettings.h
@@ -0,0 +1,65 @@
+﻿// Copyright Epic Games, Inc. All Rights Reserved.
+
+#pragma once
+
+#include "CoreMinimal.h"
+#include "Engine/DeveloperSettings.h"
+#include "WebAPILiquidJSSettings.generated.h"
+
+/** */
+UCLASS(BlueprintType, Config = Engine, DefaultConfig, meta = (DisplayName = "WebAPI LiquidJS"))
+class WEBAPILIQUIDJS_API UWebAPILiquidJSSettings : public UDeveloperSettings
+{
+	GENERATED_BODY()
+
+public:
+	virtual FName GetContainerName() const override { return FName("Project"); }
+	virtual FName GetCategoryName() const override { return FName("Plugins"); }
+
+	/** The web app http port. */
+	UPROPERTY(Config, EditAnywhere, AdvancedDisplay, Category = "Service")
+	uint32 Port = 33000;
+
+	/** Should force a build of the WebApp at startup. */
+	UPROPERTY(Config, EditAnywhere, AdvancedDisplay, Category = "Service")
+	bool bForceWebAppBuildAtStartup = false;
+
+	/** Should WebApp log timing. */
+	UPROPERTY(config, EditAnywhere, AdvancedDisplay, Category = "Service")
+	bool bWebAppLogRequestDuration = false;
+
+	/** Whether web server is started automatically. */
+	UPROPERTY(config, EditAnywhere, AdvancedDisplay, Category = "Service")
+	bool bAutoStartWebServer = true;
+
+	/** Whether web socket server is started automatically. */
+	UPROPERTY(config, EditAnywhere, AdvancedDisplay, Category = "Service")
+	bool bAutoStartWebSocketServer = true;
+
+	/** The HTTP server's port. */
+	UPROPERTY(config, EditAnywhere, AdvancedDisplay, Category = "Service")
+	uint32 HttpServerPort = 33010;
+
+	/** The WebSocket server's port. */
+	UPROPERTY(config, EditAnywhere, AdvancedDisplay, Category = "Service")
+	uint32 WebSocketServerPort = 33020;
+
+	/** Returns a formatted Url. */
+	FString GetServiceUrl(const FString& InSubPath = {});
+
+	/** Returns a formatted Url. */
+	FString GetServiceUrl(const FString& InSubPath = {}) const;
+
+private:
+	/** The web app url. */
+	UPROPERTY(Config, EditAnywhere, BlueprintReadWrite, Category = "Service", meta = (AllowPrivateAccess))
+	FString ServiceUrl = TEXT("127.0.0.1:{Port}");
+	
+	/** Cached formatted Url. */
+	UPROPERTY(Transient)
+	FString FormattedServiceUrl;
+
+	/** What port the formatted Url was created with, used to re-cache on port change. */
+	UPROPERTY(Transient)
+	uint32 FormattedWithPort = 0;
+};
diff --git a/Engine/Plugins/Web/WebAPI/Source/WebAPILiquidJS/Private/WebAPILiquidJSStats.cpp b/Engine/Plugins/Web/WebAPI/Source/WebAPILiquidJS/Private/WebAPILiquidJSStats.cpp
new file mode 100644
--- /dev/null
+++ b/Engine/Plugins/Web/WebAPI/Source/WebAPILiquidJS/Private/WebAPILiquidJSStats.cpp
@@ -0,0 +1,3 @@
+// Copyright Epic Games, Inc. All Rights Reserved.
+
+#include "WebAPILiquidJSStats.h"
diff --git a/Engine/Plugins/Web/WebAPI/Source/WebAPILiquidJS/Private/WebAPILiquidJSStats.h b/Engine/Plugins/Web/WebAPI/Source/WebAPILiquidJS/Private/WebAPILiquidJSStats.h
new file mode 100644
--- /dev/null
+++ b/Engine/Plugins/Web/WebAPI/Source/WebAPILiquidJS/Private/WebAPILiquidJSStats.h
@@ -0,0 +1,10 @@
+// Copyright Epic Games, Inc. All Rights Reserved.
+
+#pragma once
+
+#include "CoreMinimal.h"
+
+#include "Stats/Stats.h"
+#include "Stats/Stats2.h"
+
+DECLARE_STATS_GROUP(TEXT("WebAPILiquidJS"), STATGROUP_WebAPILiquidJS, STATCAT_Advanced);
diff --git a/Engine/Plugins/Web/WebAPI/Source/WebAPILiquidJS/Public/IWebAPILiquidJSModule.h b/Engine/Plugins/Web/WebAPI/Source/WebAPILiquidJS/Public/IWebAPILiquidJSModule.h
new file mode 100644
--- /dev/null
+++ b/Engine/Plugins/Web/WebAPI/Source/WebAPILiquidJS/Public/IWebAPILiquidJSModule.h
@@ -0,0 +1,11 @@
+// Copyright Epic Games, Inc. All Rights Reserved.
+
+#pragma once
+
+#include "CoreMinimal.h"
+#include "Modules/ModuleInterface.h"
+
+class IWebAPILiquidJSModuleInterface
+    : public IModuleInterface
+{
+};
diff --git a/Engine/Plugins/Web/WebAPI/Source/WebAPILiquidJS/WebAPILiquidJS.Build.cs b/Engine/Plugins/Web/WebAPI/Source/WebAPILiquidJS/WebAPILiquidJS.Build.cs
new file mode 100644
--- /dev/null
+++ b/Engine/Plugins/Web/WebAPI/Source/WebAPILiquidJS/WebAPILiquidJS.Build.cs
@@ -0,0 +1,40 @@
+// Copyright Epic Games, Inc. All Rights Reserved.
+
+using UnrealBuildTool;
+
+public class WebAPILiquidJS : ModuleRules
+{
+	public WebAPILiquidJS(ReadOnlyTargetRules Target) : base(Target)
+	{
+		PCHUsage = ModuleRules.PCHUsageMode.UseExplicitOrSharedPCHs;
+
+		PublicDependencyModuleNames.AddRange(
+			new string[] {
+				"AssetTools",
+				"Core",
+				"CoreUObject",
+				"Engine",
+				"Json",
+				"JsonUtilities",
+			});
+
+		PrivateDependencyModuleNames.AddRange(
+			new string[] {
+				"AssetRegistry",
+				"DeveloperSettings",
+				"EditorStyle",
+				"EditorWidgets",
+				"Engine",
+				"HTTP",
+				"InputCore",
+				"Projects",
+				"Slate",
+				"SlateCore",
+				"Sockets",
+				"UnrealEd",
+				"WebAPI",
+				"WebAPIEditor",
+				"WebSocketNetworking",
+			});
+	}
+}
diff --git a/Engine/Plugins/Web/WebAPI/WebAPIGeneratorApp/_package.json b/Engine/Plugins/Web/WebAPI/WebAPIGeneratorApp/_package.json
new file mode 100644
--- /dev/null
+++ b/Engine/Plugins/Web/WebAPI/WebAPIGeneratorApp/_package.json
@@ -0,0 +1,40 @@
+{
+  "name": "web-api-generator-server",
+  "version": "0.0.1",
+  "description": "WebAPIGeneratorApp initialised with create-express-typescript-application.",
+  "main": "src/index.ts",
+  "scripts": {
+    "build": "tsc",
+    "start": "node ./build/index.js",
+    "dev": "ts-node ./src/index.ts",
+    "dev:nodemon": "tsc && nodemon -w src -w templates -e ts,json,liquid -x ts-node ./src/index.ts",
+    "find-deadcode": "ts-prune"
+  },
+  "keywords": ["express", "typescript"],
+  "devDependencies": {
+    "@types/express": "^4.17.13",
+    "@types/node": "^14.17.14",
+    "body-parser": "^1.19.0",
+    "express": "^4.17.1",
+    "fs-extra": "^9.0.1",
+    "lodash": "^4.17.20",
+    "nodemon": "^2.0.12",
+    "superagent": "^6.0.0",
+    "ts-node": "^9.1.1",
+    "typescript": "^4.4.2",
+    "ws": "^7.3.1"
+  },
+  "dependencies": {
+    "@overnightjs/core": "^1.7.6",
+    "@types/express": "^4.17.7",
+    "@types/fs-extra": "^9.0.1",
+    "@types/lodash": "^4.14.160",
+    "@types/node": "^14.6.0",
+    "@types/superagent": "4.1.10",
+    "@types/ws": "^7.2.6",  
+    "express": "^4.17.1",
+    "http-status": "^1.5.0",
+    "liquidjs": "^9.25.1",
+    "typescript": "^4.0.2"    
+  }
+}
diff --git a/Engine/Plugins/Web/WebAPI/WebAPIGeneratorApp/scripts/build.js b/Engine/Plugins/Web/WebAPI/WebAPIGeneratorApp/scripts/build.js
new file mode 100644
--- /dev/null
+++ b/Engine/Plugins/Web/WebAPI/WebAPIGeneratorApp/scripts/build.js
@@ -0,0 +1,31 @@
+const fs = require('fs-extra');
+const path = require('path');
+const exec = require('child_process').exec;
+const util = require('util');
+const package = require('../package.json');
+
+const execute = util.promisify(exec);
+
+async function build() {
+  try {
+    const root = path.resolve(__dirname, '..');
+
+    await Promise.all([
+      execute('npm run build', { cwd: root })
+    ])
+
+    await fs.writeJSON(path.join(root, 'public', 'version.json'), { date: new Date(), version: package.version });
+
+  } catch (err) {
+    let message = err.stdout;
+    if (!message)
+      message = err.message;
+    if (!message)
+      message = err;
+
+    console.log('Build failed: ', message);
+    process.exit(1);
+  }
+}
+
+build();
diff --git a/Engine/Plugins/Web/WebAPI/WebAPIGeneratorApp/scripts/start.js b/Engine/Plugins/Web/WebAPI/WebAPIGeneratorApp/scripts/start.js
new file mode 100644
--- /dev/null
+++ b/Engine/Plugins/Web/WebAPI/WebAPIGeneratorApp/scripts/start.js
@@ -0,0 +1,108 @@
+const fs = require('fs');
+const path = require('path');
+const cp = require('child_process');
+
+const root = path.resolve(__dirname, '..');
+
+function execute(command, cwd, description, output) {
+  return new Promise((resolve, reject) => {
+    const child = cp.exec(command, { encoding: 'utf8', cwd });
+    child.addListener('error', reject);
+    child.addListener('exit', code => {
+      if (code === 0)
+        resolve();
+      else
+        reject(new Error(`Failed To ${description}`));
+    });
+
+    if (output)
+      child.stdout.on('data', console.log);
+  });
+}
+
+function printError(err) {
+  let message = err.stdout;
+  if (!message)
+    message = err.message;
+  if (!message)
+    message = err;
+
+  console.log('ERROR:', message);
+}
+
+async function build() {
+  try {
+    const build = path.join(root, 'build');
+    const installed = path.join(build, 'version.txt');
+    const pkgPath = path.join(root, '_package.json');
+
+    if (!fs.existsSync(pkgPath)) {
+      console.log("ERROR: Failed to build WebApp - Can't find package.json")
+      return false;
+    }
+
+    const pkg = JSON.parse(fs.readFileSync(pkgPath, 'utf8'));
+    if (!pkg || !pkg.version) {
+      console.log("ERROR: Failed to build WebApp - Can't parse package.json")
+      return false;
+    }
+
+	const forceBuild = !!process.argv.find(arg => arg === '--build');
+	if (!forceBuild && fs.existsSync(installed)) {
+      const version = fs.readFileSync(installed, 'utf8');
+      if (version && pkg.version === version)
+        return true;
+    }
+
+    for (const folder of [root]) {
+      for (const filename of ['package.json', 'package-lock.json']) {
+        const source = path.join(folder, `_${filename}`);
+        const target = path.join(folder, filename);
+        if (!fs.existsSync(source))
+          continue;
+
+        if (fs.existsSync(target))
+          fs.unlinkSync(target);
+
+        fs.copyFileSync(source, target);
+        const stat = fs.statSync(target);
+
+        // Remove readonly flag set by perforce
+        fs.chmodSync(target, stat.mode | 0o600);
+      }
+    }
+
+    console.log('Installing dependencies');
+    await execute('npm install', root, 'Install dependencies');
+
+    console.log('Building WebApp');
+    await execute('npm run build', root, 'Build WebApp', true);
+    if (!fs.existsSync(build))
+        return false;
+
+    fs.writeFileSync(installed, pkg.version, 'utf8');
+    return true;
+
+  } catch (err) {
+    printError(err);
+    return false;
+  }
+}
+
+async function start() {
+  try {
+    console.log('Starting WebApp...');
+    const args = process.argv.slice(2).map(arg => `"${arg}"`).join(' ');
+    const compiled = path.join(root, 'build');
+    await execute(`node "${compiled}" ${args}`, root, 'Run WebApp', true);
+  } catch (err) {
+    printError(err);
+  }
+}
+
+async function buildAndStart() {
+  if (await build())
+    await start();
+}
+
+buildAndStart();
diff --git a/Engine/Plugins/Web/WebAPI/WebAPIGeneratorApp/src/controllers/api.controller.ts b/Engine/Plugins/Web/WebAPI/WebAPIGeneratorApp/src/controllers/api.controller.ts
new file mode 100644
--- /dev/null
+++ b/Engine/Plugins/Web/WebAPI/WebAPIGeneratorApp/src/controllers/api.controller.ts
@@ -0,0 +1,49 @@
+import { Controller, Get } from '@overnightjs/core';
+import { Request, Response } from 'express';
+
+interface ISendOptions {
+  type?: string;
+  nullIsNotFound?: boolean;
+}
+
+@Controller('api')
+export class ApiController {
+  public async initialize(): Promise<void> {
+  }
+
+  @Get('ping')
+  private ping(req: Request, res: Response) {
+    res.send({ message: 'ok' });
+  }
+
+  @Get('shutdown')
+  private shutdown(req: Request, res: Response) {
+    res.send({ message: 'ok' });
+    setImmediate(() => process.exit(0));
+  }
+
+  @Get('status')
+  private status(req: Request, res: Response) {
+    res.send({ message: 'ok' });
+  }
+
+  protected async send(res: Response, promise: Promise<any>, options?: ISendOptions): Promise<any> {
+    try {
+      const result = await promise;
+      if (result === null && options?.nullIsNotFound) {
+        return res.status(404).send({ status: 'ERROR', error: 'Not found' });
+      }
+
+      if (options?.type) {
+        res.set('Content-Type', options.type);
+      }
+
+      res.send(result);
+
+    } catch (err: any) {
+      const error = err?.message ?? err;
+      res.status(500)
+          .send({ status: 'ERROR', error });
+    }
+  }
+}
diff --git a/Engine/Plugins/Web/WebAPI/WebAPIGeneratorApp/src/controllers/generator.controller.ts b/Engine/Plugins/Web/WebAPI/WebAPIGeneratorApp/src/controllers/generator.controller.ts
new file mode 100644
--- /dev/null
+++ b/Engine/Plugins/Web/WebAPI/WebAPIGeneratorApp/src/controllers/generator.controller.ts
@@ -0,0 +1,78 @@
+import { Controller, Post } from '@overnightjs/core';
+import { Request, Response } from 'express';
+const catchAsync_ = require('../utils/catchAsync');
+const { generatorService } = require('../services');
+const httpStatus = require('http-status');
+
+interface ISendOptions {
+    type?: string;
+    nullIsNotFound?: boolean;
+}
+
+@Controller('api/gen')
+export class GeneratorController {
+    public async initialize(): Promise<void> {
+    }
+
+    @Post('struct')
+    private async struct(req: Request, res: Response) {
+        const generatedStruct = await generatorService.generateStruct(req.body);
+        res.status(httpStatus.OK).send(generatedStruct);
+    }
+
+    @Post('enum')
+    private async enum(req: Request, res: Response) {
+        const generatedEnum = await generatorService.generateEnum(req.body);
+        res.status(httpStatus.OK).send(generatedEnum);
+    }
+
+    @Post('header')
+    private async header(req: Request, res: Response) {
+        const generatedHeader = await generatorService.generateHeader(req.body);
+        res.status(httpStatus.OK).send(generatedHeader);
+    }
+
+    @Post('call')
+    private async call(req: Request, res: Response) {
+        const generatedCall = await generatorService.generateCall(req.body);
+        res.status(httpStatus.OK).send(generatedCall);
+    }
+
+    @Post('asyncAction')
+    private async asyncAction(req: Request, res: Response) {
+        const generatedAsyncAction = await generatorService.generateAsyncAction(req.body);
+        res.status(httpStatus.OK).send(generatedAsyncAction);
+    }
+
+    @Post('decl')
+    private async decl(req: Request, res: Response) {
+        const generatedDecl = await generatorService.generateDecl(req.body);
+        res.status(httpStatus.OK).send(generatedDecl);
+    }
+
+    @Post('defn')
+    private async defn(req: Request, res: Response) {
+        const generatedDefn = await generatorService.generateDefn(req.body);
+        res.status(httpStatus.OK).send(generatedDefn);
+    }
+
+    protected async send(res: Response, promise: Promise<any>, options?: ISendOptions): Promise<any> {
+        try {
+            const result = await promise;
+            if (result === null && options?.nullIsNotFound) {
+                return res.status(404).send({ status: 'ERROR', error: 'Not found' });
+            }
+
+            if (options?.type) {
+                res.set('Content-Type', options.type);
+            }
+
+            res.send(result);
+
+        } catch (err: any) {
+            const error = err?.message ?? err;
+            res.status(500)
+                .send({ status: 'ERROR', error });
+        }
+    }
+}
diff --git a/Engine/Plugins/Web/WebAPI/WebAPIGeneratorApp/src/controllers/generator2.controller.ts b/Engine/Plugins/Web/WebAPI/WebAPIGeneratorApp/src/controllers/generator2.controller.ts
new file mode 100644
--- /dev/null
+++ b/Engine/Plugins/Web/WebAPI/WebAPIGeneratorApp/src/controllers/generator2.controller.ts
@@ -0,0 +1,55 @@
+// The underscore is required because TypeScript is stupid
+const catchAsync_ = require("../utils/catchAsync");
+const { generatorService } = require("../services");
+const httpStatus = require("http-status");
+
+const requestStruct = catchAsync_(async (req: any, res: any) => {
+    const generatedStruct = await generatorService.generateStruct(req.body);
+    res.status(httpStatus.OK).send(generatedStruct);
+});
+
+const requestEnum = catchAsync_(async (req: any, res: any) => {
+    const generatedEnum = await generatorService.generateEnum(req.body);
+    res.status(httpStatus.OK).send(generatedEnum);
+});
+
+const requestHeader = catchAsync_(async (req: any, res: any) => {
+    const generatedHeader = await generatorService.generateHeader(req.body);
+    res.status(httpStatus.OK).send(generatedHeader);
+});
+
+const requestCall = catchAsync_(async (req: any, res: any) => {
+    const generatedCall = await generatorService.generateCall(req.body);
+    res.status(httpStatus.OK).send(generatedCall);
+});
+
+const requestSubsystem = catchAsync_(async (req: any, res: any) => {
+    const generatedSubsystem = await generatorService.generateSubsystem({});
+    res.status(httpStatus.OK).send(generatedSubsystem);
+});
+
+const requestAsyncAction = catchAsync_(async (req: any, res: any) => {
+    const generatedAsyncAction = await generatorService.generateAsyncAction(req.body);
+    res.status(httpStatus.OK).send(generatedAsyncAction);
+});
+
+const requestDecl = catchAsync_(async (req: any, res: any) => {
+    const generatedDecl = await generatorService.generateDecl(req.body);
+    res.status(httpStatus.OK).send(generatedDecl);
+});
+
+const requestDefn = catchAsync_(async (req: any, res: any) => {
+    const generatedDefn = await generatorService.generateDefn(req.body);
+    res.status(httpStatus.OK).send(generatedDefn);
+});
+
+module.exports = {
+    requestStruct,
+    requestEnum,
+    requestHeader,
+    requestCall,
+    requestSubsystem,
+    requestAsyncAction,
+    requestDecl,
+    requestDefn,
+};
diff --git a/Engine/Plugins/Web/WebAPI/WebAPIGeneratorApp/src/filters/caseConverters.ts b/Engine/Plugins/Web/WebAPI/WebAPIGeneratorApp/src/filters/caseConverters.ts
new file mode 100644
--- /dev/null
+++ b/Engine/Plugins/Web/WebAPI/WebAPIGeneratorApp/src/filters/caseConverters.ts
@@ -0,0 +1,21 @@
+// Usage: {{ items | titlecase }}
+import { MetadataInfo } from '../types/metadata';
+
+export function titleCase(v: string): string {
+    if(!v || v.length === 0) { return ""; }
+
+    let prefix = "";
+    if(v.startsWith("(")) {
+        prefix = "(";
+        v = v.substring(1);
+    }
+
+    let suffix = "";
+    if(v.endsWith(")")) {
+        suffix = ")";
+        v = v.substring(0, v.length - 1);
+    }
+
+    const result = v.replace(/([A-Z])/g, " $1").trim();
+    return prefix + result.charAt(0).toUpperCase() + result.slice(1) + suffix;
+}
diff --git a/Engine/Plugins/Web/WebAPI/WebAPIGeneratorApp/src/filters/defaultValueFilter.ts b/Engine/Plugins/Web/WebAPI/WebAPIGeneratorApp/src/filters/defaultValueFilter.ts
new file mode 100644
--- /dev/null
+++ b/Engine/Plugins/Web/WebAPI/WebAPIGeneratorApp/src/filters/defaultValueFilter.ts
@@ -0,0 +1,33 @@
+import { MetadataInfo } from '../types/metadata';
+import { PropertyInfo } from '../types/property';
+import { NameVariant, TypeInfo, toString } from '../types/typeInfo';
+
+/** Returns the default value for the given object */
+// Usage: {{ item | defaultvalue }}
+export function defaultValue(v: PropertyInfo): string {
+    if(!v) {
+        return "";
+    }
+
+    if(v.defaultValue)
+    {
+        return v.defaultValue
+    }
+
+    if(typeof v.type != 'string')
+    {
+        if(!v.type.isBuiltinType)
+        {
+            if(v.type.prefix == "F")
+            {
+                return `${toString(v.type)}::GetDefault()`;
+            }
+            else if(v.type.prefix == "U")
+            {
+                return "nullptr;"
+            }
+        }
+    }
+
+    return '{}'
+}
diff --git a/Engine/Plugins/Web/WebAPI/WebAPIGeneratorApp/src/filters/delegateNumFilter.ts b/Engine/Plugins/Web/WebAPI/WebAPIGeneratorApp/src/filters/delegateNumFilter.ts
new file mode 100644
--- /dev/null
+++ b/Engine/Plugins/Web/WebAPI/WebAPIGeneratorApp/src/filters/delegateNumFilter.ts
@@ -0,0 +1,12 @@
+// Usage: {{ num | numStr }}
+export function numStr(v: number): string {
+    if(v == 0 || v > 8) {
+        return "";
+    }
+
+    const strLookup: Array<string> = [
+        "One", "Two", "Three", "Four", "Five", "Six", "Seven", "Eight"
+    ];
+
+    return strLookup[v - 1] + (v > 1 ? "Params" : "Param");
+}
diff --git a/Engine/Plugins/Web/WebAPI/WebAPIGeneratorApp/src/filters/kvpFilter.ts b/Engine/Plugins/Web/WebAPI/WebAPIGeneratorApp/src/filters/kvpFilter.ts
new file mode 100644
--- /dev/null
+++ b/Engine/Plugins/Web/WebAPI/WebAPIGeneratorApp/src/filters/kvpFilter.ts
@@ -0,0 +1,17 @@
+// Usage: {{ items | kvp %}}
+import { MetadataInfo } from '../types/metadata';
+
+export function keyValuePairs(v: Array<MetadataInfo | string>, prefix?: string, suffix?: string): string {
+    if(!v || v.length === 0) { return ""; }
+    const keyValuePairs = v.map((v: MetadataInfo | string) => {
+        if(typeof v === "string") {
+            return `${v}`;
+        }
+        else {
+            return "value" in v && v.value ? `${v.key}="${v.value}"` : `${v.key}`;
+        }
+    });
+    prefix = prefix ?? "";
+    suffix = suffix ?? "";
+    return `${prefix}${keyValuePairs.join(", ")}${suffix}`;
+}
diff --git a/Engine/Plugins/Web/WebAPI/WebAPIGeneratorApp/src/filters/paramsFilter.ts b/Engine/Plugins/Web/WebAPI/WebAPIGeneratorApp/src/filters/paramsFilter.ts
new file mode 100644
--- /dev/null
+++ b/Engine/Plugins/Web/WebAPI/WebAPIGeneratorApp/src/filters/paramsFilter.ts
@@ -0,0 +1,26 @@
+// Usage: {% parameters | params %}
+import { OperationParameterInfo } from '../types/operation';
+import { DelegateInfo } from '../types/delegate';
+
+// mode = undefined, decl, defn, macro or call
+export function params(v: Array<OperationParameterInfo | DelegateInfo>, mode?: string): string {
+    mode = mode ?? 'decl';
+    const keyValuePairs = v.map((v: OperationParameterInfo | DelegateInfo) => {
+        if('type' in v) {
+            let prefix = mode === 'decl' ? `UPARAM(DisplayName = "${v.name}") ` : "";
+            let suffix = mode === 'decl' ? v.defaultValue ? ` = ${v.defaultValue}` : "" : "";
+            return mode === 'call'
+                ?  `${v.name}`
+                : mode === 'macro'
+                ? `${v.isConst ? "const " : ""}${v.type}${v.accessModifier ?? ""}, ${v.name}`
+                : `${prefix}${v.isConst ? "const " : ""}${v.type}${v.accessModifier ?? ""} In${v.name}${suffix}`;
+        }
+        else {
+            let prefix = mode === 'decl' ? `UPARAM(DisplayName = "On${v.name}") ` : "";
+            return mode === 'call'
+                ? `InOn${v.name}`
+                : `${prefix}FOn${v.name}Delegate& InOn${v.name}`;
+        }
+    });
+    return keyValuePairs.join(", ");
+}
diff --git a/Engine/Plugins/Web/WebAPI/WebAPIGeneratorApp/src/filters/setFilter.ts b/Engine/Plugins/Web/WebAPI/WebAPIGeneratorApp/src/filters/setFilter.ts
new file mode 100644
--- /dev/null
+++ b/Engine/Plugins/Web/WebAPI/WebAPIGeneratorApp/src/filters/setFilter.ts
@@ -0,0 +1,9 @@
+// Sets the specified property (or adds it) of one or more objects, to the specified value
+
+// Usage: {{ items | set: "propertyName", 55 }}
+export function set(arr: {[key: string]: any}[], property: string, value: any): {[key: string]: any}[] {
+    arr.forEach((obj: {[key: string]: any}) => {
+        obj[property] = value;
+    });
+    return arr;
+}
diff --git a/Engine/Plugins/Web/WebAPI/WebAPIGeneratorApp/src/filters/structToDelegateFilter.ts b/Engine/Plugins/Web/WebAPI/WebAPIGeneratorApp/src/filters/structToDelegateFilter.ts
new file mode 100644
--- /dev/null
+++ b/Engine/Plugins/Web/WebAPI/WebAPIGeneratorApp/src/filters/structToDelegateFilter.ts
@@ -0,0 +1,23 @@
+// Usage: {% parameters | params %}
+import { OperationParameterInfo, OperationResponseInfo } from '../types/operation';
+import { DelegateInfo, DelegateParameterInfo } from '../types/delegate';
+import { StructInfo } from '../types/struct';
+import { PropertyInfo } from '../types/property';
+
+// mode = undefined, decl, defn, macro or call
+export function structToDelegate(v: Array<StructInfo | OperationResponseInfo>): DelegateInfo[] {
+    if(!v) {
+        return [];
+    }
+    return v.map((obj: StructInfo | OperationResponseInfo) => {
+        const parameters = obj.properties
+            ? obj.properties.map((p: PropertyInfo) => {
+                return p as DelegateParameterInfo;
+            })
+            : [];
+        return <DelegateInfo>{
+            name: obj.name,
+            parameters: parameters
+        };
+    });
+}
diff --git a/Engine/Plugins/Web/WebAPI/WebAPIGeneratorApp/src/filters/typeInfoFilter.ts b/Engine/Plugins/Web/WebAPI/WebAPIGeneratorApp/src/filters/typeInfoFilter.ts
new file mode 100644
--- /dev/null
+++ b/Engine/Plugins/Web/WebAPI/WebAPIGeneratorApp/src/filters/typeInfoFilter.ts
@@ -0,0 +1,45 @@
+import { MetadataInfo } from '../types/metadata';
+import { NameVariant, TypeInfo, toString } from '../types/typeInfo';
+
+// Usage: {{ item | kvp }}
+export function toStr(v: NameVariant | TypeInfo, justName?: boolean): string {
+    return toString(v, justName);
+}
+
+// Usage: {{ item | jsonname }}
+export function jsonName(v: NameVariant | TypeInfo): string {
+    if(!v) {
+        return "";
+    }
+    else if(typeof v == 'string') {
+        return v;
+    }
+    else {
+        return v.jsonName ?? v.name;
+    }
+}
+
+// Usage: {{ item | jsontype }}
+export function jsonType(v: NameVariant | TypeInfo): string {
+    if(!v) {
+        return "";
+    }
+
+    let value = v;
+    if(typeof v == 'string') {
+        value = v;
+    }
+    else {
+        value = v.jsonType ?? v.name;
+    }
+
+    const validValues = [ "null", "string", "number", "boolean", "bool", "array", "object" ];
+
+    if(validValues.indexOf(value.toLowerCase()) < 0)
+    {
+        // Not a valid json type (might be "ToFromHandler") so return empty string
+        value = "";
+    }
+
+    return value;
+}
diff --git a/Engine/Plugins/Web/WebAPI/WebAPIGeneratorApp/src/filters/whereFilter.ts b/Engine/Plugins/Web/WebAPI/WebAPIGeneratorApp/src/filters/whereFilter.ts
new file mode 100644
--- /dev/null
+++ b/Engine/Plugins/Web/WebAPI/WebAPIGeneratorApp/src/filters/whereFilter.ts
@@ -0,0 +1,44 @@
+import { isTruthy } from "liquidjs";
+
+// Extension of the default "where" filter, that allows an operator other than "===", and a defaultValue to compare against if the property resolves to undefined
+export function where<T extends object>(arr: T[], property: string, expected?: any, operator?: string, defaultValue?: any): T[] {
+    if(!arr) {
+        return [];
+    }
+
+    let filtered = arr.filter((obj: T) => {
+        const objWithKeys: {[key: string]: any} = obj
+        let propertyParts = property.split('.')
+        let value = objWithKeys
+        for (let index = 0; index < propertyParts.length; index++) {
+            const propertyPart = propertyParts[index]
+            value = value[propertyPart]
+            if(value === undefined && defaultValue !== undefined) {
+                value = defaultValue
+            }
+        }
+
+        if(value === undefined) {
+            return false;
+        }
+
+        if(expected !== undefined) {
+            operator = operator || "==="
+            const expression = (left: any, right: any) => {
+                let str = `return function(left, right) {
+                    "use strict";
+                    return (left ${operator} right);
+                }`
+                let functor = Function(str)
+                return functor()
+            };
+            const expressionResult = expression(value, expected)(value, expected)
+            return expressionResult
+        }
+        else {
+            return value ? true : false
+        }
+    })
+
+    return filtered;
+}
diff --git a/Engine/Plugins/Web/WebAPI/WebAPIGeneratorApp/src/logServer.ts b/Engine/Plugins/Web/WebAPI/WebAPIGeneratorApp/src/logServer.ts
new file mode 100644
--- /dev/null
+++ b/Engine/Plugins/Web/WebAPI/WebAPIGeneratorApp/src/logServer.ts
@@ -0,0 +1,69 @@
+import ws from 'ws';
+import _ from 'lodash';
+import fs from 'fs-extra';
+import { ServiceSettings } from './settings';
+
+export namespace LogServer {
+    var server: ws.Server | null;
+    var buffer: string = '';
+    var startTime = new Date().toISOString().replace(/\:/g, '_');
+    var flushTimer: NodeJS.Timeout | null;
+
+    export async function initialize(): Promise<void> {
+        const port = ServiceSettings.dev ? 7002 : ServiceSettings.port + 2;
+        server = new ws.Server({ port });
+        await fs.ensureDir('./logs');
+        server.on('connection', (socket: ws) => {
+            socket.on('message', onMessage);
+        });
+
+        flushTimer = setInterval(flusLog, 2 * 60 * 1000);
+    }
+
+    export async function shutdown(): Promise<void> {
+        server?.removeAllListeners();
+        server?.close();
+        server = null;
+        if(flushTimer) {
+            clearInterval(flushTimer);
+        }
+        flushTimer = null;
+    }
+
+    function onMessage(data: ws.Data) {
+        const json = data.toString();
+        try {
+            const payload = JSON.parse(json);
+            log(payload);
+        } catch (err) {
+            console.log("Error with message:", json);
+            console.log(err);
+        }
+    }
+
+    export function log(payload: Record<string, any>): void {
+        const Timestamp = new Date().toISOString();
+        buffer += JSON.stringify({ Timestamp, ...payload }) + '\r\n';
+    }
+
+    function flusLog() {
+        if (flushLogFile(buffer, 'all'))
+            buffer = '';
+    }
+
+    function flushLogFile(content: string, type: string) {
+        if (!content)
+            return false;
+
+        try {
+            const filename = `./Logs/${startTime}-${type}.log`;
+            fs.appendFileSync(filename, content, 'utf8');
+
+            return true;
+        } catch (err) {
+            console.log('error', err);
+        }
+
+        return false;
+    }
+}
diff --git a/Engine/Plugins/Web/WebAPI/WebAPIGeneratorApp/src/routes/v1/generators.routes.js b/Engine/Plugins/Web/WebAPI/WebAPIGeneratorApp/src/routes/v1/generators.routes.js
new file mode 100644
--- /dev/null
+++ b/Engine/Plugins/Web/WebAPI/WebAPIGeneratorApp/src/routes/v1/generators.routes.js
@@ -0,0 +1,13 @@
+const express = require('express');
+const generatorController = require('../../controllers/generator.controller');
+const router = express.Router();
+
+router.route(`/struct`).post(generatorController.requestStruct);
+router.route('/enum').post(generatorController.requestEnum);
+router.route('/header').post(generatorController.requestHeader);
+router.route('/call').post(generatorController.requestCall);
+router.route('/asyncAction').post(generatorController.requestAsyncAction);
+router.route('/decl').post(generatorController.requestDecl);
+router.route('/defn').post(generatorController.requestDefn);
+
+module.exports = router;
diff --git a/Engine/Plugins/Web/WebAPI/WebAPIGeneratorApp/src/routes/v1/index.js b/Engine/Plugins/Web/WebAPI/WebAPIGeneratorApp/src/routes/v1/index.js
new file mode 100644
--- /dev/null
+++ b/Engine/Plugins/Web/WebAPI/WebAPIGeneratorApp/src/routes/v1/index.js
@@ -0,0 +1,17 @@
+const express = require('express');
+const generatorRoutes = require('./generators.routes');
+
+const router = express.Router();
+
+const defaultRoutes = [
+    {
+        path: '/gen',
+        route: generatorRoutes
+    }
+];
+
+defaultRoutes.forEach(route => {
+    router.use(route.path, route.route);
+});
+
+module.exports = router;
\ No newline at end of file
diff --git a/Engine/Plugins/Web/WebAPI/WebAPIGeneratorApp/src/services/generators.service.ts b/Engine/Plugins/Web/WebAPI/WebAPIGeneratorApp/src/services/generators.service.ts
new file mode 100644
--- /dev/null
+++ b/Engine/Plugins/Web/WebAPI/WebAPIGeneratorApp/src/services/generators.service.ts
@@ -0,0 +1,90 @@
+const getTemplateRoute = require("../utils/getTemplateRoute");
+const generateMetadata = require("../utils/generateMetadata");
+
+import { Liquid } from "liquidjs";
+import { titleCase } from "../filters/caseConverters";
+import { keyValuePairs } from "../filters/kvpFilter";
+import { params } from "../filters/paramsFilter";
+import { jsonName, toStr, jsonType } from "../filters/typeInfoFilter";
+import { OperationInfo, OperationResponseInfo } from "../types/operation";
+import { DeclInfo } from "../types/decl";
+import { EnumInfo } from "../types/enum";
+import { HeaderInfo } from "../types/header";
+import { StructInfo } from "../types/struct";
+import { where } from "../filters/whereFilter";
+import { structToDelegate } from "../filters/structToDelegateFilter";
+import { numStr } from "../filters/delegateNumFilter";
+import { set } from "../filters/setFilter";
+import { defaultValue } from "../filters/defaultValueFilter";
+
+// @fixme: duplicated in server.ts
+const engine: Liquid = new Liquid({
+    root: getTemplateRoute,
+    extname: ".liquid",
+    // strictFilters: true,
+    // strictVariables: true
+});
+
+engine.registerFilter('kvp', keyValuePairs);
+engine.registerFilter('params', params)
+engine.registerFilter('str', toStr)
+engine.registerFilter('jsonname', jsonName)
+engine.registerFilter('jsontype', jsonType)
+engine.registerFilter('titlecase', titleCase)
+engine.registerFilter('structToDelegate', structToDelegate)
+engine.registerFilter('where', where)
+engine.registerFilter('numStr', numStr)
+engine.registerFilter('set', set)
+engine.registerFilter('defaultvalue', defaultValue)
+
+const generateStruct = async (body: StructInfo) => {
+    return await engine.renderFile("structDecl", { 'structDecl': body });
+};
+
+const generateEnum = async (body: EnumInfo) => {
+    return await engine.renderFile("enumDecl", body);
+};
+
+const generateHeader = async (body: HeaderInfo) => {
+    return await engine.renderFile("header", body);
+};
+
+const generateCall = async (body: OperationInfo) => {
+    return await engine.renderFile("callDecl", body);
+};
+
+const generateAsyncAction = async (body: OperationInfo) => {
+    let decl = await engine.renderFile("asyncActionDecl", body);
+    let defn = await engine.renderFile("asyncActionDefn", body);
+    return defn;
+    //return [decl, defn]
+    // return await engine.renderFile("asyncActionDecl", body);
+};
+
+const generateDecl = async (body: DeclInfo) => {
+    body.includePaths.sort();
+    // @todo: remove after testing
+    body.operations.forEach(operation => {
+        operation.logName = operation.logName || "LogTemp";
+    });
+    return await engine.renderFile("header", body);
+};
+
+const generateDefn = async (body: DeclInfo) => {
+    body.includePaths.sort();
+    // @todo: remove after testing
+    body.operations.forEach(operation => {
+        operation.logName = operation.logName || "LogTemp";
+    });
+    return await engine.renderFile("cpp", body);
+};
+
+module.exports = {
+    generateStruct,
+    generateEnum,
+    generateHeader,
+    generateCall,
+    generateAsyncAction,
+    generateDecl,
+    generateDefn,
+};
diff --git a/Engine/Plugins/Web/WebAPI/WebAPIGeneratorApp/src/services/index.ts b/Engine/Plugins/Web/WebAPI/WebAPIGeneratorApp/src/services/index.ts
new file mode 100644
--- /dev/null
+++ b/Engine/Plugins/Web/WebAPI/WebAPIGeneratorApp/src/services/index.ts
@@ -0,0 +1 @@
+module.exports.generatorService = require('./generators.service');
diff --git a/Engine/Plugins/Web/WebAPI/WebAPIGeneratorApp/src/settings.ts b/Engine/Plugins/Web/WebAPI/WebAPIGeneratorApp/src/settings.ts
new file mode 100644
--- /dev/null
+++ b/Engine/Plugins/Web/WebAPI/WebAPIGeneratorApp/src/settings.ts
@@ -0,0 +1,61 @@
+import path from 'path';
+
+class ServiceSettings {
+  readonly name: string = 'WebAPILiquidJS';
+  readonly version: string = 'v1';
+  readonly dev: boolean;
+  readonly monitor: boolean;
+  readonly port: number;
+  readonly ueHttpPort: number;
+  readonly ueWebSocketPort: number;
+  readonly rootFolder: string;
+  readonly logger: boolean;
+
+  constructor() {
+    this.dev = false;
+    this.monitor = false;
+    this.port = 33000;
+    this.ueHttpPort = 33010;
+    this.ueWebSocketPort = 33020;
+    this.logger = false;
+
+    for (let i = 2; i < process.argv.length; i++) {
+      switch (process.argv[i]) {
+        case '--dev':
+          this.dev = true;
+          this.port = 33001;
+          this.logger = true;
+          break;
+
+        case '--monitor':
+          this.monitor = true;
+          break;
+
+        case '--port':
+          this.port = parseInt(process.argv[i + 1]) || 33000;
+          i++;
+          break;
+
+        case '--uehttp':
+          this.ueHttpPort = parseInt(process.argv[i + 1]) || 33010;
+          i++;
+          break;
+
+        case '--uews':
+          this.ueWebSocketPort = parseInt(process.argv[i + 1]) || 33020;
+          i++;
+          break;
+
+        case '--log':
+          this.logger = true;
+          break;
+      }
+    }
+
+    process.env.NODE_ENV = this.dev ? 'development' : 'production';
+    this.rootFolder = path.resolve('./');
+  }
+}
+
+const Instance = new ServiceSettings();
+export { Instance as ServiceSettings };
\ No newline at end of file
diff --git a/Engine/Plugins/Web/WebAPI/WebAPIGeneratorApp/src/types/class.ts b/Engine/Plugins/Web/WebAPI/WebAPIGeneratorApp/src/types/class.ts
new file mode 100644
--- /dev/null
+++ b/Engine/Plugins/Web/WebAPI/WebAPIGeneratorApp/src/types/class.ts
@@ -0,0 +1,15 @@
+import { FunctionInfo } from "./function";
+import { MetadataInfo } from "./metadata";
+import { PropertyInfo } from "./property";
+import { NameVariant } from "./typeInfo";
+
+export interface ClassInfo {
+    description: string;
+    name: NameVariant;
+    namespace: string;
+    module: string;
+    specifiers: Array<MetadataInfo> | string;
+    metadata: Array<MetadataInfo> | string;
+    properties: Array<PropertyInfo>;
+    functions: Array<FunctionInfo>;
+}
diff --git a/Engine/Plugins/Web/WebAPI/WebAPIGeneratorApp/src/types/decl.ts b/Engine/Plugins/Web/WebAPI/WebAPIGeneratorApp/src/types/decl.ts
new file mode 100644
--- /dev/null
+++ b/Engine/Plugins/Web/WebAPI/WebAPIGeneratorApp/src/types/decl.ts
@@ -0,0 +1,26 @@
+import { ClassInfo } from "./class";
+import { EnumInfo } from "./enum";
+import { OperationInfo, OperationRequestInfo, OperationResponseInfo } from "./operation";
+import { SettingsClassInfo } from "./settings";
+import { StructInfo } from "./struct";
+
+/** Wraps multiple types in a single header. */
+export interface DeclInfo {
+    baseFilePath?: string;
+    relativeFilePath: string;
+    fileName: string;
+    fileType: string;
+    namespace: string;
+    copyrightNotice: string;
+    module: string;
+    moduleDependencies?: Array<string>;
+    includePaths: Array<string>;
+    forwardDeclarations?: Array<string>;
+    enums: Array<EnumInfo>;
+    structs: Array<StructInfo>;
+    requests: Array<OperationRequestInfo>;
+    responses: Array<OperationResponseInfo>;
+    classes: Array<ClassInfo>;
+    settings?: Array<SettingsClassInfo>;
+    operations: Array<OperationInfo>
+}
diff --git a/Engine/Plugins/Web/WebAPI/WebAPIGeneratorApp/src/types/delegate.ts b/Engine/Plugins/Web/WebAPI/WebAPIGeneratorApp/src/types/delegate.ts
new file mode 100644
--- /dev/null
+++ b/Engine/Plugins/Web/WebAPI/WebAPIGeneratorApp/src/types/delegate.ts
@@ -0,0 +1,14 @@
+import { PropertyInfo } from "./property";
+import { NameVariant } from "./typeInfo";
+
+export interface DelegateParameterInfo extends PropertyInfo {
+    isConst?: boolean; // if unspecified, false
+    accessModifier?: string // * or &
+}
+
+export interface DelegateInfo {
+    name: NameVariant;
+    parameters: Array<DelegateParameterInfo>;
+    isDynamic?: boolean;
+    isMulticast?: boolean;
+}
diff --git a/Engine/Plugins/Web/WebAPI/WebAPIGeneratorApp/src/types/enum.ts b/Engine/Plugins/Web/WebAPI/WebAPIGeneratorApp/src/types/enum.ts
new file mode 100644
--- /dev/null
+++ b/Engine/Plugins/Web/WebAPI/WebAPIGeneratorApp/src/types/enum.ts
@@ -0,0 +1,17 @@
+import { NameVariant } from "./typeInfo";
+
+interface EnumValueInfo {
+    name: NameVariant;
+    jsonName?: NameVariant;
+    displayName?: NameVariant;
+    description?: string;
+    explicitValue?: Number;
+}
+
+export interface EnumInfo {
+    description: string;
+    name: NameVariant;
+    namespace: string;
+    module: string;
+    values: Array<EnumValueInfo | string>;
+}
diff --git a/Engine/Plugins/Web/WebAPI/WebAPIGeneratorApp/src/types/function.ts b/Engine/Plugins/Web/WebAPI/WebAPIGeneratorApp/src/types/function.ts
new file mode 100644
--- /dev/null
+++ b/Engine/Plugins/Web/WebAPI/WebAPIGeneratorApp/src/types/function.ts
@@ -0,0 +1,27 @@
+import { MetadataInfo } from "./metadata";
+import { PropertyInfo } from "./property";
+import { NameVariant } from "./typeInfo";
+
+interface FunctionParameterInfo extends PropertyInfo {
+    isConst?: boolean; // if unspecified, false
+    accessModifier?: string // * or &
+}
+
+export interface FunctionInfo {
+    name: NameVariant;
+    description: string;
+    namespace: string;
+    module: string;
+    parameters: Array<FunctionParameterInfo>;
+    isOverride?: boolean;
+    isConst?: boolean;
+
+    returnType: string;
+    isReturnTypeConst?: boolean; // if unspecified, false
+    body: string;
+
+    specifiers: Array<MetadataInfo> | string;
+    metadata: Array<MetadataInfo> | string;
+
+    includes: Array<string>;
+}
diff --git a/Engine/Plugins/Web/WebAPI/WebAPIGeneratorApp/src/types/header.ts b/Engine/Plugins/Web/WebAPI/WebAPIGeneratorApp/src/types/header.ts
new file mode 100644
--- /dev/null
+++ b/Engine/Plugins/Web/WebAPI/WebAPIGeneratorApp/src/types/header.ts
@@ -0,0 +1,6 @@
+export interface HeaderInfo {
+    includes: Array<string>;
+    forwardDeclarations: Array<string>;
+    copyrightNotice: string;
+    filenameWithoutExtension: string;
+}
diff --git a/Engine/Plugins/Web/WebAPI/WebAPIGeneratorApp/src/types/metadata.ts b/Engine/Plugins/Web/WebAPI/WebAPIGeneratorApp/src/types/metadata.ts
new file mode 100644
--- /dev/null
+++ b/Engine/Plugins/Web/WebAPI/WebAPIGeneratorApp/src/types/metadata.ts
@@ -0,0 +1,6 @@
+// metadata is a key/value pair, or just a Key with no value
+// ex: { key: "DisplayName", value: "MyFunction" }
+export interface MetadataInfo {
+    key: string
+    value?: string
+}
diff --git a/Engine/Plugins/Web/WebAPI/WebAPIGeneratorApp/src/types/operation.ts b/Engine/Plugins/Web/WebAPI/WebAPIGeneratorApp/src/types/operation.ts
new file mode 100644
--- /dev/null
+++ b/Engine/Plugins/Web/WebAPI/WebAPIGeneratorApp/src/types/operation.ts
@@ -0,0 +1,54 @@
+import { DelegateInfo } from "./delegate";
+import { StructInfo } from "./struct";
+import { NameVariant } from "./typeInfo";
+
+enum ParameterStorage {
+    Path,
+    Query,
+    Header,
+    Cookie,
+    Body
+}
+
+export interface OperationParameterInfo {
+    name: string;
+    type: string;
+    description: string;
+    isConst?: boolean; // if unspecified, false
+    accessModifier?: string // * or &
+    defaultValue?: any;
+    storage: ParameterStorage;
+    mediaType?: string;
+}
+
+export interface OperationRequestInfo extends StructInfo {
+    body?: StructInfo;
+}
+
+export interface OperationResponseInfo extends StructInfo {
+    responseCode: number;
+    isArray: boolean;
+    message?: string;
+}
+
+export interface OperationInfo {
+    description: string;
+    namespace: string;
+    service: string;
+    path: string;
+    module: string;
+    name: NameVariant;
+    verb: string;
+    returnType: string;
+    errorType: string;
+    delegates?: Array<DelegateInfo>;
+    parameters: Array<OperationParameterInfo>;
+    requests: Array<OperationRequestInfo>
+    responses: Array<OperationResponseInfo>
+
+    includes: Array<string>;
+    filenameWithoutExtension: string;
+    settingsTypeName?: string;
+
+    logName?: string;
+}
diff --git a/Engine/Plugins/Web/WebAPI/WebAPIGeneratorApp/src/types/property.ts b/Engine/Plugins/Web/WebAPI/WebAPIGeneratorApp/src/types/property.ts
new file mode 100644
--- /dev/null
+++ b/Engine/Plugins/Web/WebAPI/WebAPIGeneratorApp/src/types/property.ts
@@ -0,0 +1,15 @@
+import { MetadataInfo } from "./metadata";
+import { NameVariant } from "./typeInfo";
+
+export interface PropertyInfo {
+    name: NameVariant;
+    description: string;
+    namespace: string;
+    type: NameVariant;
+    isArray: boolean;
+    isRequired: boolean;
+    isMixin?: boolean;
+    defaultValue?: string;
+    specifiers: Array<MetadataInfo> | string;
+    metadata: Array<MetadataInfo> | string;
+};
diff --git a/Engine/Plugins/Web/WebAPI/WebAPIGeneratorApp/src/types/settings.ts b/Engine/Plugins/Web/WebAPI/WebAPIGeneratorApp/src/types/settings.ts
new file mode 100644
--- /dev/null
+++ b/Engine/Plugins/Web/WebAPI/WebAPIGeneratorApp/src/types/settings.ts
@@ -0,0 +1,16 @@
+import { NameVariant } from "./typeInfo";
+
+export interface SettingsClassInfo {
+    name: NameVariant;
+    base?: NameVariant;
+    displayName?: string;
+    host: string;
+    baseUrl?: string;
+    userAgent: string;
+    dateTimeFormat?: string;
+    schemes: Array<string>;
+}
+
+// const a = { hat: "h"};
+// const b = { hat: "w"};
+// const obj: SettingsClassInfo = {...a, ...b};
diff --git a/Engine/Plugins/Web/WebAPI/WebAPIGeneratorApp/src/types/struct.ts b/Engine/Plugins/Web/WebAPI/WebAPIGeneratorApp/src/types/struct.ts
new file mode 100644
--- /dev/null
+++ b/Engine/Plugins/Web/WebAPI/WebAPIGeneratorApp/src/types/struct.ts
@@ -0,0 +1,14 @@
+import { MetadataInfo } from "./metadata";
+import { PropertyInfo } from "./property";
+import { NameVariant } from "./typeInfo";
+
+export interface StructInfo {
+    description: string;
+    name: NameVariant;
+    namespace: string;
+    module: string;
+    base?: NameVariant;
+    specifiers: Array<MetadataInfo> | string;
+    metadata: Array<MetadataInfo> | string;
+    properties: Array<PropertyInfo>;
+}
diff --git a/Engine/Plugins/Web/WebAPI/WebAPIGeneratorApp/src/types/typeInfo.ts b/Engine/Plugins/Web/WebAPI/WebAPIGeneratorApp/src/types/typeInfo.ts
new file mode 100644
--- /dev/null
+++ b/Engine/Plugins/Web/WebAPI/WebAPIGeneratorApp/src/types/typeInfo.ts
@@ -0,0 +1,60 @@
+export interface TypeInfo {
+    name: string;
+    jsonName: string;
+    jsonType: string;
+    jsonPropertyToSerialize?: string;
+    printFormatSpecifier?: string;
+    printFormatExpression?: string;
+    namespace: string;
+    prefix: string;
+    containerType: string;
+    isBuiltinType: boolean;
+
+    toString(): string;
+}
+
+export type NameVariant = string | TypeInfo;
+
+export function toString<T>(obj: T, justName?: boolean): string;
+
+export function toString(info: NameVariant, justName?: boolean): string {
+    if(!info) {
+        return "";
+    }
+    else if(typeof info == 'string') {
+        return info;
+    }
+    else {
+        if(justName) {
+            return info.name;
+        }
+
+        let containerType = info.containerType?.toLowerCase();
+        if(!info.name && containerType.includes('enum')) {
+            return "EUnknownEnum";
+        }
+
+        let typeStr = !info.namespace || info.isBuiltinType
+            ? `${info.prefix}${info.name}`
+            : `${info.prefix}${info.namespace}_${info.name}`;
+
+        if(containerType) {
+            let outerType = "";
+            if(containerType.includes('array')) {
+                outerType = "TArray";
+            }
+            else if(containerType.includes('set')) {
+                outerType = "TSet";
+            }
+            else if(containerType.includes('map')) {
+                outerType = "TMap";
+            }
+
+            if(outerType) {
+                typeStr = `${outerType}<${typeStr}>`;
+            }
+        }
+
+        return typeStr;
+    }
+}
diff --git a/Engine/Plugins/Web/WebAPI/WebAPIGeneratorApp/src/utils/catchAsync.ts b/Engine/Plugins/Web/WebAPI/WebAPIGeneratorApp/src/utils/catchAsync.ts
new file mode 100644
--- /dev/null
+++ b/Engine/Plugins/Web/WebAPI/WebAPIGeneratorApp/src/utils/catchAsync.ts
@@ -0,0 +1,5 @@
+const catchAsync = (callback: any) => (req: Request, res: Response, next: any) => {
+    Promise.resolve(callback(req, res, next)).catch((error: any) => next(error));
+};
+
+module.exports = catchAsync;
\ No newline at end of file
diff --git a/Engine/Plugins/Web/WebAPI/WebAPIGeneratorApp/src/utils/generateMetadata.ts b/Engine/Plugins/Web/WebAPI/WebAPIGeneratorApp/src/utils/generateMetadata.ts
new file mode 100644
--- /dev/null
+++ b/Engine/Plugins/Web/WebAPI/WebAPIGeneratorApp/src/utils/generateMetadata.ts
@@ -0,0 +1,16 @@
+import { MetadataInfo } from '../types/metadata';
+
+const generateMetadata = (metadataMap: Array<MetadataInfo>): string => {
+    let metadata_: string = "";
+    metadataMap?.forEach((metadata: MetadataInfo) => {
+        if (metadata_?.length !== 0) {
+            metadata_ = `${metadata_}, ${metadata.key} = "${metadata.value}"`;
+        } else {
+            metadata_ = `meta = (${metadata.key} = "${metadata.value}"`;
+        }
+    })
+
+    return `${metadata_})`;
+}
+
+module.exports = generateMetadata;
\ No newline at end of file
diff --git a/Engine/Plugins/Web/WebAPI/WebAPIGeneratorApp/src/utils/getTemplateRoute.ts b/Engine/Plugins/Web/WebAPI/WebAPIGeneratorApp/src/utils/getTemplateRoute.ts
new file mode 100644
--- /dev/null
+++ b/Engine/Plugins/Web/WebAPI/WebAPIGeneratorApp/src/utils/getTemplateRoute.ts
@@ -0,0 +1,9 @@
+import path from 'path';
+
+const getTemplateRoute = () => {
+    let root = __dirname;
+    const srcIdx = root.indexOf('src');    
+    return path.resolve(`${root.substring(0, srcIdx)}`, 'templates/');
+}
+
+module.exports = getTemplateRoute();
diff --git a/Engine/Plugins/Web/WebAPI/WebAPIGeneratorApp/start.bat b/Engine/Plugins/Web/WebAPI/WebAPIGeneratorApp/start.bat
new file mode 100755
--- /dev/null
+++ b/Engine/Plugins/Web/WebAPI/WebAPIGeneratorApp/start.bat
@@ -0,0 +1,23 @@
+@echo off
+
+rem Add default nodejs installation folder to path, in case it was not added / overwritten
+SET PATH=%PATH%;%ProgramFiles%\nodejs\
+
+rem Check if nodejs is in the env variable PATH
+for %%X in (node.exe) do (set node=%%~$PATH:X)
+if not defined node (
+  echo ERROR: Couldn't find node.js installed, Please install latest nodejs from https://nodejs.org/en/download/
+  exit 1
+)
+
+rem Let's check if it is a modern nodejs
+node -e "process.exit( process.versions.node.split('.')[0] );"
+echo Found Node.js version %errorlevel% (%node%)
+
+if %errorlevel% LSS 14 (
+  echo ERROR: installed node.js version is too old, please install latest nodejs from https://nodejs.org/en/download/
+  exit 1
+)
+
+rem redirecting all command line arguments to node script
+node scripts/start.js %*
\ No newline at end of file
diff --git a/Engine/Plugins/Web/WebAPI/WebAPIGeneratorApp/start.command b/Engine/Plugins/Web/WebAPI/WebAPIGeneratorApp/start.command
new file mode 100644
--- /dev/null
+++ b/Engine/Plugins/Web/WebAPI/WebAPIGeneratorApp/start.command
@@ -0,0 +1,7 @@
+#!/bin/sh
+# Copyright Epic Games, Inc. All Rights Reserved.
+#
+# Simple wrapper around start.sh using the
+# .command extension enables it to be run from the OSX Finder.
+
+sh "`dirname "$0"`"/start.sh $*
diff --git a/Engine/Plugins/Web/WebAPI/WebAPIGeneratorApp/start.sh b/Engine/Plugins/Web/WebAPI/WebAPIGeneratorApp/start.sh
new file mode 100644
--- /dev/null
+++ b/Engine/Plugins/Web/WebAPI/WebAPIGeneratorApp/start.sh
@@ -0,0 +1,28 @@
+#!/bin/sh
+
+set -e
+
+# First we check if nodejs is installed
+if ! command -v node > /dev/null ; then
+  echo "ERROR: Couldn't find node.js installed..., Please install latest nodejs from https://nodejs.org/en/download/"
+  exit 1
+fi
+
+# Check if npm is installed
+if ! command -v npm > /dev/null ; then
+  echo "ERROR: Couldn't find npm installed..., Please install npm"
+  exit 1
+fi
+
+# Let's check if it is a modern nodejs
+VERSION=$(node -e "console.log( process.versions.node.split('.')[0] );")
+echo "Found Node.js version ${VERSION}"
+
+if [ ${VERSION} -lt 14 ] ; then
+  echo "ERROR: installed node.js version is too old (${VERSION}) :\( Please install latest nodejs from https://nodejs.org/en/download/"
+  exit 1
+fi
+
+FOLDER=$(dirname "$0")
+node ${FOLDER}/scripts/start.js "$@"
+
