WebAPI missing files pt2

diff --git a/Engine/Plugins/Web/WebAPI/Source/WebAPIOpenAPI/Private/Factories/WebAPIOpenAPIFactoryBase.cpp b/Engine/Plugins/Web/WebAPI/Source/WebAPIOpenAPI/Private/Factories/WebAPIOpenAPIFactoryBase.cpp
new file mode 100644
--- /dev/null
+++ b/Engine/Plugins/Web/WebAPI/Source/WebAPIOpenAPI/Private/Factories/WebAPIOpenAPIFactoryBase.cpp
@@ -0,0 +1,31 @@
+﻿// Copyright Epic Games, Inc. All Rights Reserved.
+
+#include "WebAPIOpenAPIFactoryBase.h"
+
+#define LOCTEXT_NAMESPACE "WebAPIOpenAPIFactoryBase"
+
+const FName UWebAPIOpenAPIFactoryBase::JsonFileType = FName(TEXT("json"));
+const FName UWebAPIOpenAPIFactoryBase::YamlFileType = FName(TEXT("yaml"));
+
+UWebAPIOpenAPIFactoryBase::UWebAPIOpenAPIFactoryBase()
+{
+	Formats.Add(TEXT("json;JavaScript Object Notation"));
+}
+
+TFuture<bool> UWebAPIOpenAPIFactoryBase::ImportWebAPI(UWebAPIDefinition* InDefinition, const FString& InFileName, const FString& InFileContents)
+{
+	// @todo: modify definition here
+	// write raw spec contents
+	// parse spec
+
+	// @todo: return false if json couldn't be parsed
+
+	return MakeFulfilledPromise<bool>(false).GetFuture();
+}
+
+bool UWebAPIOpenAPIFactoryBase::IsValidFileExtension(const FString& InFileExtension) const
+{
+	return InFileExtension.Equals(JsonFileType.ToString(), ESearchCase::IgnoreCase);
+}
+
+#undef LOCTEXT_NAMESPACE
diff --git a/Engine/Plugins/Web/WebAPI/Source/WebAPIOpenAPI/Private/Factories/WebAPIOpenAPIFactoryBase.h b/Engine/Plugins/Web/WebAPI/Source/WebAPIOpenAPI/Private/Factories/WebAPIOpenAPIFactoryBase.h
new file mode 100644
--- /dev/null
+++ b/Engine/Plugins/Web/WebAPI/Source/WebAPIOpenAPI/Private/Factories/WebAPIOpenAPIFactoryBase.h
@@ -0,0 +1,26 @@
+﻿// Copyright Epic Games, Inc. All Rights Reserved.
+
+#pragma once
+
+#include "CoreMinimal.h"
+#include "Factories/WebAPIDefinitionFactory.h"
+
+#include "WebAPIOpenAPIFactoryBase.generated.h"
+
+/** Common base factory between Swagger and OpenAPI factories. */
+UCLASS(Abstract)
+class WEBAPIOPENAPI_API UWebAPIOpenAPIFactoryBase
+	: public UWebAPIDefinitionFactory
+{
+	GENERATED_BODY()
+
+public:
+	UWebAPIOpenAPIFactoryBase();
+
+	virtual TFuture<bool> ImportWebAPI(UWebAPIDefinition* InDefinition, const FString& InFileName, const FString& InFileContents) override;
+	virtual bool IsValidFileExtension(const FString& InFileExtension) const override;
+
+protected:
+	static const FName JsonFileType;
+	static const FName YamlFileType;
+};
diff --git a/Engine/Plugins/Web/WebAPI/Source/WebAPIOpenAPI/Private/Tests/Samples/V3/ContactObjectExample.json b/Engine/Plugins/Web/WebAPI/Source/WebAPIOpenAPI/Private/Tests/Samples/V3/ContactObjectExample.json
new file mode 100644
--- /dev/null
+++ b/Engine/Plugins/Web/WebAPI/Source/WebAPIOpenAPI/Private/Tests/Samples/V3/ContactObjectExample.json
@@ -0,0 +1,5 @@
+﻿{
+  "name": "API Support",
+  "url": "http://www.example.com/support",
+  "email": "[REDACTED_EMAIL]"
+}
\ No newline at end of file
diff --git a/Engine/Plugins/Web/WebAPI/Source/WebAPIOpenAPI/Private/Tests/Samples/V3/removebg.json b/Engine/Plugins/Web/WebAPI/Source/WebAPIOpenAPI/Private/Tests/Samples/V3/removebg.json
new file mode 100644
--- /dev/null
+++ b/Engine/Plugins/Web/WebAPI/Source/WebAPIOpenAPI/Private/Tests/Samples/V3/removebg.json
@@ -0,0 +1,984 @@
+{
+  "openapi": "3.0.0",
+  "servers": [
+    {
+      "url": "https://api.remove.bg/v1.0"
+    }
+  ],
+  "info": {
+    "description": "Remove the background of any image",
+    "title": "Background Removal API",
+    "version": "1.0.0",
+    "x-apisguru-categories": [
+      "developer_tools"
+    ],
+    "x-logo": {
+      "url": "https://api.apis.guru/v2/cache/logo/https_remove.bg_apple-touch-icon.png"
+    },
+    "x-origin": [
+      {
+        "format": "openapi",
+        "url": "https://www.remove.bg/api/swagger.yaml",
+        "version": "3.0"
+      }
+    ],
+    "x-providerName": "remove.bg"
+  },
+  "security": [
+    {
+      "APIKeyHeader": []
+    }
+  ],
+  "paths": {
+    "/account": {
+      "get": {
+        "description": "Get the current credit balance and number of free API calls.\n\nNotes:\n\n* Balance changes may be delayed by several seconds. To locally keep track of your credit balance, you should therefore only call this endpoint initially (or e.g. when the user manually triggers a refresh), then use the `X-Credits-Charged` response header returned with each image processing response to adjust the local balance.\n\n* The \"*sizes*\" field is always \"all\", is deprecated and will soon be removed.\n",
+        "responses": {
+          "200": {
+            "content": {
+              "*/*": {
+                "schema": {
+                  "properties": {
+                    "data": {
+                      "properties": {
+                        "attributes": {
+                          "properties": {
+                            "api": {
+                              "properties": {
+                                "free_calls": {
+                                  "example": 50,
+                                  "type": "integer"
+                                },
+                                "sizes": {
+                                  "example": "all",
+                                  "type": "string"
+                                }
+                              },
+                              "type": "object"
+                            },
+                            "credits": {
+                              "properties": {
+                                "enterprise": {
+                                  "example": 0,
+                                  "format": "float",
+                                  "type": "number"
+                                },
+                                "payg": {
+                                  "example": 50,
+                                  "format": "float",
+                                  "type": "number"
+                                },
+                                "subscription": {
+                                  "example": 150,
+                                  "format": "float",
+                                  "type": "number"
+                                },
+                                "total": {
+                                  "example": 200,
+                                  "format": "float",
+                                  "type": "number"
+                                }
+                              },
+                              "type": "object"
+                            }
+                          },
+                          "type": "object"
+                        }
+                      },
+                      "type": "object"
+                    }
+                  },
+                  "type": "object"
+                }
+              }
+            },
+            "description": ""
+          },
+          "403": {
+            "content": {
+              "*/*": {
+                "schema": {
+                  "$ref": "#/components/schemas/AuthFailed"
+                }
+              }
+            },
+            "description": "Error: Authentication failed"
+          },
+          "429": {
+            "content": {
+              "*/*": {
+                "schema": {
+                  "$ref": "#/components/schemas/RateLimit"
+                }
+              }
+            },
+            "description": "Error: Rate limit exceeded"
+          }
+        },
+        "summary": "Fetch credit balance and free API calls.",
+        "tags": [
+          "Fetch account info"
+        ]
+      }
+    },
+    "/improve": {
+      "post": {
+        "description": "Submit an image to the remove.bg Improvement program\n* Contribute an image that remove.bg is currently not able to remove the background from properly\n* Help us make remove.bg better\n* Get better results for similiar images in the future\n\nNotes:\n  * By submitting images through the API you agree to the <a target=\"_blank\" rel=\"noopener\" href=\"/ipc\">Improvement Program Conditions</a>\n  * File size: up to 12MB\n  * up to 100 files per day (up to 1000 for Enterprise customers). <br> Higher Rate Limits are available <a href=\"/support/contact?subject=Improvement+Program+Rate+Limit\">upon request</a>.\n\nRequires either an API Key to be provided in the `X-API-Key` request header or an OAuth 2.0 access token to be provided in the `Authorization` request header.\n\nPlease note that submissions are used on a best-effort basis and the extent of expected improvement varies depending on many factors, including the number of provided images, their complexity and visual similarity. Improvements usually take several weeks to become effective.\n",
+        "requestBody": {
+          "content": {
+            "application/json": {
+              "schema": {
+                "$ref": "#/components/schemas/ImprovementProgramJson"
+              }
+            },
+            "application/x-www-form-urlencoded": {
+              "schema": {
+                "$ref": "#/components/schemas/ImprovementProgramJson"
+              }
+            },
+            "multipart/form-data": {
+              "schema": {
+                "$ref": "#/components/schemas/ImprovementProgramMultipart"
+              }
+            }
+          },
+          "required": true
+        },
+        "responses": {
+          "200": {
+            "content": {
+              "*/*": {
+                "schema": {
+                  "$ref": "#/components/schemas/ImprovementProgramJsonResponse"
+                }
+              }
+            },
+            "description": "Image submitted"
+          },
+          "400": {
+            "content": {
+              "*/*": {
+                "schema": {
+                  "properties": {
+                    "errors": {
+                      "items": {
+                        "properties": {
+                          "title": {
+                            "description": "Error message",
+                            "example": "File too large",
+                            "type": "string"
+                          }
+                        },
+                        "type": "object"
+                      },
+                      "type": "array"
+                    }
+                  },
+                  "type": "object"
+                }
+              }
+            },
+            "description": "Error: Invalid parameters or input file unprocessable"
+          },
+          "403": {
+            "content": {
+              "*/*": {
+                "schema": {
+                  "$ref": "#/components/schemas/AuthFailed"
+                }
+              }
+            },
+            "description": "Error: Authentication failed"
+          },
+          "429": {
+            "content": {
+              "*/*": {
+                "schema": {
+                  "$ref": "#/components/schemas/RateLimit"
+                }
+              }
+            },
+            "description": "Error: Rate limit exceeded"
+          }
+        },
+        "tags": [
+          "Improvement Program"
+        ]
+      }
+    },
+    "/removebg": {
+      "post": {
+        "description": "Removes the background of a JPG/PNG image.\n\n* File size: up to 12 MB\n* Image source: File upload (binary or as base64 encoded string) or download from URL\n* Image Content: Any photo with a foreground [(e.g. people, products, animals, cars, etc.)](/supported-images)\n* Output resolutions available: Preview (up to 0.25 megapixels), Full (up to 25 megapixels)\n\nRequires either an API Key to be provided in the `X-API-Key` request header or an OAuth 2.0 access token to be provided in the `Authorization` request header.\n",
+        "requestBody": {
+          "content": {
+            "application/json": {
+              "schema": {
+                "$ref": "#/components/schemas/RemoveBgJson"
+              }
+            },
+            "application/x-www-form-urlencoded": {
+              "schema": {
+                "$ref": "#/components/schemas/RemoveBgJson"
+              }
+            },
+            "multipart/form-data": {
+              "schema": {
+                "$ref": "#/components/schemas/RemoveBgMultipart"
+              }
+            }
+          },
+          "required": true
+        },
+        "responses": {
+          "200": {
+            "content": {
+              "application/json": {
+                "schema": {
+                  "$ref": "#/components/schemas/RemoveBgJsonResponse"
+                }
+              },
+              "image/*": {
+                "schema": {
+                  "example": "(binary image data)",
+                  "format": "binary",
+                  "type": "string"
+                }
+              }
+            },
+            "description": "Image background removed",
+            "headers": {
+              "X-Credits-Charged": {
+                "description": "Amount of credits charged for this call",
+                "schema": {
+                  "format": "float",
+                  "type": "number"
+                }
+              },
+              "X-Height": {
+                "description": "Height of the result image",
+                "schema": {
+                  "type": "integer"
+                }
+              },
+              "X-Type": {
+                "description": "Detected foreground type (How specific this classification is depends on the type_level parameter sent in the request)",
+                "schema": {
+                  "enum": [
+                    "person",
+                    "product",
+                    "animal",
+                    "car",
+                    "car_interior",
+                    "car_part",
+                    "transportation",
+                    "graphics",
+                    "other"
+                  ],
+                  "type": "string"
+                }
+              },
+              "X-Width": {
+                "description": "Width of the result image",
+                "schema": {
+                  "type": "integer"
+                }
+              }
+            }
+          },
+          "400": {
+            "content": {
+              "*/*": {
+                "examples": {
+                  "failed_bg_download": {
+                    "value": {
+                      "errors": [
+                        {
+                          "code": "failed_bg_download",
+                          "title": "Failed to download background image from given image_url: https://unsplash.com/"
+                        }
+                      ]
+                    }
+                  },
+                  "failed_image_download": {
+                    "value": {
+                      "errors": [
+                        {
+                          "code": "failed_image_download",
+                          "title": "Failed to download image from given image_url: https://unsplash.com/"
+                        }
+                      ]
+                    }
+                  },
+                  "failed_to_read_image": {
+                    "value": {
+                      "errors": [
+                        {
+                          "code": "failed_to_read_image",
+                          "title": "There was an error reading the image."
+                        }
+                      ]
+                    }
+                  },
+                  "file_too_large": {
+                    "value": {
+                      "errors": [
+                        {
+                          "code": "file_too_large",
+                          "detail": "File exceeds limit of 12MB",
+                          "title": "File too large"
+                        }
+                      ]
+                    }
+                  },
+                  "invalid_bg_color": {
+                    "value": {
+                      "errors": [
+                        {
+                          "code": "invalid_bg_color",
+                          "title": "Invalid bg_color parameter given"
+                        }
+                      ]
+                    }
+                  },
+                  "invalid_bg_image_url": {
+                    "value": {
+                      "errors": [
+                        {
+                          "code": "invalid_bg_image_url",
+                          "title": "Invalid bg_image_url: Please provide a valid URL."
+                        }
+                      ]
+                    }
+                  },
+                  "invalid_channels": {
+                    "value": {
+                      "errors": [
+                        {
+                          "code": "invalid_channels",
+                          "title": "Invalid value for parameter 'channels'"
+                        }
+                      ]
+                    }
+                  },
+                  "invalid_crop_margin": {
+                    "value": {
+                      "errors": [
+                        {
+                          "code": "invalid_crop_margin",
+                          "title": "Invalid crop_margin parameter given"
+                        }
+                      ]
+                    }
+                  },
+                  "invalid_dimensions": {
+                    "value": {
+                      "errors": [
+                        {
+                          "code": "invalid_dimensions",
+                          "detail": "The dimensions of the given image could not be read.",
+                          "title": "Failed to read image dimensions"
+                        }
+                      ]
+                    }
+                  },
+                  "invalid_file_type": {
+                    "value": {
+                      "errors": [
+                        {
+                          "code": "invalid_file_type",
+                          "detail": "Expected jpg/png, received mp4",
+                          "title": "Invalid file type"
+                        }
+                      ]
+                    }
+                  },
+                  "invalid_format": {
+                    "value": {
+                      "errors": [
+                        {
+                          "code": "invalid_format",
+                          "title": "Invalid format parameter given"
+                        }
+                      ]
+                    }
+                  },
+                  "invalid_image_url": {
+                    "value": {
+                      "errors": [
+                        {
+                          "code": "invalid_image_url",
+                          "title": "Invalid image_url: Please provide a valid URL."
+                        }
+                      ]
+                    }
+                  },
+                  "invalid_position": {
+                    "value": {
+                      "errors": [
+                        {
+                          "code": "invalid_position",
+                          "title": "Invalid position parameter given"
+                        }
+                      ]
+                    }
+                  },
+                  "invalid_roi": {
+                    "value": {
+                      "errors": [
+                        {
+                          "code": "invalid_roi",
+                          "title": "Invalid roi parameter given"
+                        }
+                      ]
+                    }
+                  },
+                  "invalid_scale": {
+                    "value": {
+                      "errors": [
+                        {
+                          "code": "invalid_scale",
+                          "title": "Invalid scale parameter given"
+                        }
+                      ]
+                    }
+                  },
+                  "invalid_semitransparency": {
+                    "value": {
+                      "errors": [
+                        {
+                          "code": "invalid_semitransparency",
+                          "title": "Invalid semitransparency parameter given"
+                        }
+                      ]
+                    }
+                  },
+                  "invalid_size": {
+                    "value": {
+                      "errors": [
+                        {
+                          "code": "invalid_size",
+                          "title": "Invalid value for parameter 'size'"
+                        }
+                      ]
+                    }
+                  },
+                  "missing_source": {
+                    "value": {
+                      "errors": [
+                        {
+                          "code": "missing_source",
+                          "detail": "Please provide the source image in the image_url, image_file or image_file_b64 parameter.",
+                          "title": "No image given"
+                        }
+                      ]
+                    }
+                  },
+                  "multiple_bg_sources": {
+                    "value": {
+                      "errors": [
+                        {
+                          "code": "multiple_bg_sources",
+                          "title": "Multiple background sources given: Please provide either the bg_color, the bg_image_url or the bg_image_file parameter."
+                        }
+                      ]
+                    }
+                  },
+                  "multiple_sources": {
+                    "value": {
+                      "errors": [
+                        {
+                          "code": "multiple_sources",
+                          "title": "Multiple image sources given: Please provide either the image_url, image_file or image_file_b64 parameter."
+                        }
+                      ]
+                    }
+                  },
+                  "resolution_too_high": {
+                    "value": {
+                      "errors": [
+                        {
+                          "code": "resolution_too_high",
+                          "detail": "Input image has 60 megapixels, maximum supported input resolution is 50 megapixels",
+                          "title": "Image resolution too high"
+                        }
+                      ]
+                    }
+                  },
+                  "roi_exceeds_bounds": {
+                    "value": {
+                      "errors": [
+                        {
+                          "code": "roi_exceeds_bounds",
+                          "detail": "The given roi parameter defines a region that exceeds the image bounds",
+                          "title": "ROI exceeds image bounds"
+                        }
+                      ]
+                    }
+                  },
+                  "roi_region_empty": {
+                    "value": {
+                      "errors": [
+                        {
+                          "code": "roi_region_empty",
+                          "detail": "The given roi parameter defines an empty region",
+                          "title": "ROI region is empty"
+                        }
+                      ]
+                    }
+                  },
+                  "unknown_foreground": {
+                    "value": {
+                      "errors": [
+                        {
+                          "code": "unknown_foreground",
+                          "title": "Could not identify foreground in image. For details and recommendations see https://www.remove.bg/supported-images."
+                        }
+                      ]
+                    }
+                  }
+                },
+                "schema": {
+                  "properties": {
+                    "errors": {
+                      "items": {
+                        "properties": {
+                          "code": {
+                            "description": "Error code",
+                            "example": "file_too_large",
+                            "type": "string"
+                          },
+                          "detail": {
+                            "description": "Detailed description of the error",
+                            "example": "File exceeds limit of 12MB",
+                            "type": "string"
+                          },
+                          "title": {
+                            "description": "Error message",
+                            "example": "File too large",
+                            "type": "string"
+                          }
+                        },
+                        "type": "object"
+                      },
+                      "type": "array"
+                    }
+                  },
+                  "type": "object"
+                }
+              }
+            },
+            "description": "Error: Invalid parameters or input file unprocessable (no credits charged)"
+          },
+          "402": {
+            "content": {
+              "*/*": {
+                "schema": {
+                  "properties": {
+                    "errors": {
+                      "items": {
+                        "properties": {
+                          "title": {
+                            "description": "Error message",
+                            "example": "Insufficient credits",
+                            "type": "string"
+                          }
+                        },
+                        "type": "object"
+                      },
+                      "type": "array"
+                    }
+                  },
+                  "type": "object"
+                }
+              }
+            },
+            "description": "Error: Insufficient credits (no credits charged)"
+          },
+          "403": {
+            "content": {
+              "*/*": {
+                "schema": {
+                  "$ref": "#/components/schemas/AuthFailed"
+                }
+              }
+            },
+            "description": "Error: Authentication failed (no credits charged)"
+          },
+          "429": {
+            "content": {
+              "*/*": {
+                "schema": {
+                  "$ref": "#/components/schemas/RateLimit"
+                }
+              }
+            },
+            "description": "Error: Rate limit exceeded (no credits charged)"
+          }
+        },
+        "summary": "Remove the background of an image",
+        "tags": [
+          "Background Removal"
+        ]
+      }
+    }
+  },
+  "components": {
+    "schemas": {
+      "AuthFailed": {
+        "properties": {
+          "errors": {
+            "items": {
+              "properties": {
+                "title": {
+                  "description": "Error message",
+                  "example": "Missing API Key",
+                  "type": "string"
+                }
+              },
+              "type": "object"
+            },
+            "type": "array"
+          }
+        },
+        "type": "object"
+      },
+      "ImprovementProgramJson": {
+        "properties": {
+          "image_file_b64": {
+            "description": "Source image file (base64-encoded string). (If this parameter is present, the other image source parameters must be empty.)",
+            "example": "",
+            "type": "string"
+          },
+          "image_filename": {
+            "description": "Filename of the image, if not provided it will be autodetected form the submitted data.",
+            "example": "car.jpg",
+            "type": "string"
+          },
+          "image_url": {
+            "description": "Source image URL. (If this parameter is present, the other image source parameters must be empty.)",
+            "example": "https://www.remove.bg/example-hd.jpg",
+            "type": "string"
+          },
+          "tag": {
+            "description": "Images with the same tag are grouped together.",
+            "example": "batch_1_2020",
+            "type": "string"
+          }
+        },
+        "type": "object"
+      },
+      "ImprovementProgramJsonResponse": {
+        "properties": {
+          "id": {
+            "description": "ID of the submitted image",
+            "example": "fe09820b-be79-4511-8836-679583529882",
+            "type": "string"
+          }
+        },
+        "type": "object"
+      },
+      "ImprovementProgramMultipart": {
+        "properties": {
+          "image_file": {
+            "description": "Source image file (binary). (If this parameter is present, the other image source parameters must be empty.)",
+            "format": "binary",
+            "type": "string"
+          },
+          "image_file_b64": {
+            "description": "Source image file (base64-encoded string). (If this parameter is present, the other image source parameters must be empty.)",
+            "example": "",
+            "type": "string"
+          },
+          "image_filename": {
+            "description": "Filename of the image, if not provided it will be autodetected form the submitted data.",
+            "type": "string"
+          },
+          "image_url": {
+            "description": "Source image URL. (If this parameter is present, the other image source parameters must be empty.)",
+            "example": "https://www.remove.bg/example-hd.jpg",
+            "type": "string"
+          },
+          "tag": {
+            "description": "Images with the same tag are grouped together.",
+            "type": "string"
+          }
+        },
+        "type": "object"
+      },
+      "RateLimit": {
+        "properties": {
+          "errors": {
+            "items": {
+              "properties": {
+                "title": {
+                  "description": "Error message",
+                  "example": "Rate limit exceeded",
+                  "type": "string"
+                }
+              },
+              "type": "object"
+            },
+            "type": "array"
+          }
+        },
+        "type": "object"
+      },
+      "RemoveBgJson": {
+        "properties": {
+          "add_shadow": {
+            "default": "false",
+            "description": "Whether to add an artificial shadow to the result (default: false). NOTE: Adding shadows is currently only supported for car photos. Other subjects are returned without shadow, even if set to true (this might change in the future).\n",
+            "type": "boolean"
+          },
+          "bg_color": {
+            "description": "Adds a solid color background. Can be a hex color code (e.g. 81d4fa, fff) or a color name (e.g. green). For semi-transparency, 4-/8-digit hex codes are also supported (e.g. 81d4fa77). (If this parameter is present, the other bg_ parameters must be empty.)\n",
+            "example": "",
+            "type": "string"
+          },
+          "bg_image_url": {
+            "description": "Adds a background image from a URL. The image is centered and resized to fill the canvas while preserving the aspect ratio, unless it already has the exact same dimensions as the foreground image. (If this parameter is present, the other bg_ parameters must be empty.)",
+            "example": "",
+            "type": "string"
+          },
+          "channels": {
+            "default": "rgba",
+            "description": "Request either the finalized image (\"rgba\", default) or an alpha mask (\"alpha\"). Note: Since remove.bg also applies RGB color corrections on edges, using only the alpha mask often leads to a lower final image quality. Therefore \"rgba\" is recommended.\n",
+            "enum": [
+              "rgba",
+              "alpha"
+            ],
+            "type": "string"
+          },
+          "crop": {
+            "default": "false",
+            "description": "Whether to crop off all empty regions (default: false). Note that cropping has no effect on the amount of charged credits.\n",
+            "type": "boolean"
+          },
+          "crop_margin": {
+            "default": "0",
+            "description": "Adds a margin around the cropped subject (default: 0). Can be an absolute value (e.g. \"30px\") or relative to the subject size (e.g. \"10%\"). Can be a single value (all sides), two values (top/bottom and left/right) or four values (top, right, bottom, left). This parameter only has an effect when \"crop=true\". The maximum margin that can be added on each side is 50% of the subject dimensions or 500 pixels.\n",
+            "type": "string"
+          },
+          "format": {
+            "default": "auto",
+            "description": "Result image format:\n\"auto\" = Use PNG format if transparent regions exist, otherwise use JPG format (default),\n\"png\" = PNG format with alpha transparency,\n\"jpg\" = JPG format, no transparency,\n\"zip\" = ZIP format, contains color image and alpha matte image, supports transparency (recommended).\n",
+            "enum": [
+              "auto",
+              "png",
+              "jpg",
+              "zip"
+            ],
+            "type": "string"
+          },
+          "image_file_b64": {
+            "description": "Source image file (base64-encoded string). (If this parameter is present, the other image source parameters must be empty.)",
+            "example": "",
+            "type": "string"
+          },
+          "image_url": {
+            "description": "Source image URL. (If this parameter is present, the other image source parameters must be empty.)",
+            "example": "https://www.remove.bg/example-hd.jpg",
+            "type": "string"
+          },
+          "position": {
+            "default": "original",
+            "description": "Positions the subject within the image canvas. Can be \"original\" (default unless \"scale\" is given), \"center\" (default when \"scale\" is given) or a value from \"0%\" to \"100%\" (both horizontal and vertical) or two values (horizontal, vertical).\n",
+            "type": "string"
+          },
+          "roi": {
+            "default": "0% 0% 100% 100%",
+            "description": "Region of interest: Only contents of this rectangular region can be detected as foreground. Everything outside is considered background and will be removed. The rectangle is defined as two x/y coordinates in the format \"x1 y1 x2 y2\". The coordinates can be in absolute pixels (suffix 'px') or relative to the width/height of the image (suffix '%'). By default, the whole image is the region of interest (\"0% 0% 100% 100%\").\n",
+            "type": "string"
+          },
+          "scale": {
+            "default": "original",
+            "description": "Scales the subject relative to the total image size. Can be any value from \"10%\" to \"100%\", or \"original\" (default). Scaling the subject implies \"position=center\" (unless specified otherwise).\n",
+            "type": "string"
+          },
+          "semitransparency": {
+            "default": "true",
+            "description": "Whether to have semi-transparent regions in the result (default: true). NOTE: Semitransparency is currently only supported for car windows (this might change in the future). Other objects are returned without semitransparency, even if set to true.\n",
+            "type": "boolean"
+          },
+          "size": {
+            "default": "preview",
+            "description": "Maximum output image resolution:\n\"preview\" (default) = Resize image to 0.25 megapixels (e.g. 625×400 pixels) – 0.25 credits per image,\n\"full\" = Use original image resolution, up to 25 megapixels (e.g. 6250x4000) with formats ZIP or JPG, or up to 10 megapixels (e.g. 4000x2500) with PNG – 1 credit per image),\n\"auto\" = Use highest available resolution (based on image size and available credits).\n\nFor backwards-compatibility this parameter also accepts the values \"medium\" (up to 1.5 megapixels) and \"hd\" (up to 4 megapixels) for 1 credit per image. The value \"full\" is also available under the name \"4k\" and the value \"preview\" is aliased as \"small\" and \"regular\".\n",
+            "enum": [
+              "preview",
+              "full",
+              "auto"
+            ],
+            "type": "string"
+          },
+          "type": {
+            "default": "auto",
+            "description": "Foreground type:\n\"auto\" = Automatically detect kind of foreground,\n\"person\" = Use person(s) as foreground,\n\"product\" = Use product(s) as foreground.\n\"car\" = Use car as foreground,\n",
+            "enum": [
+              "auto",
+              "person",
+              "product",
+              "car"
+            ],
+            "type": "string"
+          },
+          "type_level": {
+            "default": 1,
+            "description": "Classification level of the detected foreground type:\n\"none\" = No classification (X-Type Header won't bet set on the response)\n\"1\" = Use coarse classification classes: [person, product, animal, car, other]\n\"2\" = Use more specific classification classes: [person, product, animal, car, car_interior, car_part, transportation, graphics, other]\n\"latest\" = Always use the latest classification classes available\n",
+            "enum": [
+              "none",
+              1,
+              2,
+              "latest"
+            ],
+            "type": "string"
+          }
+        },
+        "type": "object"
+      },
+      "RemoveBgJsonResponse": {
+        "properties": {
+          "data": {
+            "properties": {
+              "result_b64": {
+                "description": "Base64 encoded string of result image",
+                "example": "iVBORw0KGgoAAAANSUhEUgAAAIsAAACFC...",
+                "type": "string"
+              }
+            },
+            "type": "object"
+          }
+        },
+        "type": "object"
+      },
+      "RemoveBgMultipart": {
+        "properties": {
+          "add_shadow": {
+            "default": "false",
+            "description": "Whether to add an artificial shadow to the result (default: false). NOTE: Adding shadows is currently only supported for car photos. Other subjects are returned without shadow, even if set to true (this might change in the future).\n",
+            "type": "boolean"
+          },
+          "bg_color": {
+            "description": "Adds a solid color background. Can be a hex color code (e.g. 81d4fa, fff) or a color name (e.g. green). For semi-transparency, 4-/8-digit hex codes are also supported (e.g. 81d4fa77). (If this parameter is present, the other bg_ parameters must be empty.)\n",
+            "example": "",
+            "type": "string"
+          },
+          "bg_image_file": {
+            "description": "Adds a background image from a file (binary). The image is centered and resized to fill the canvas while preserving the aspect ratio, unless it already has the exact same dimensions as the foreground image. (If this parameter is present, the other bg_ parameters must be empty.)",
+            "format": "binary",
+            "type": "string"
+          },
+          "bg_image_url": {
+            "description": "Adds a background image from a URL. The image is centered and resized to fill the canvas while preserving the aspect ratio, unless it already has the exact same dimensions as the foreground image. (If this parameter is present, the other bg_ parameters must be empty.)",
+            "example": "",
+            "type": "string"
+          },
+          "channels": {
+            "default": "rgba",
+            "description": "Request either the finalized image (\"rgba\", default) or an alpha mask (\"alpha\"). Note: Since remove.bg also applies RGB color corrections on edges, using only the alpha mask often leads to a lower final image quality. Therefore \"rgba\" is recommended.\n",
+            "enum": [
+              "rgba",
+              "alpha"
+            ],
+            "type": "string"
+          },
+          "crop": {
+            "default": "false",
+            "description": "Whether to crop off all empty regions (default: false). Note that cropping has no effect on the amount of charged credits.\n",
+            "type": "boolean"
+          },
+          "crop_margin": {
+            "default": "0",
+            "description": "Adds a margin around the cropped subject (default: 0). Can be an absolute value (e.g. \"30px\") or relative to the subject size (e.g. \"10%\"). Can be a single value (all sides), two values (top/bottom and left/right) or four values (top, right, bottom, left). This parameter only has an effect when \"crop=true\". The maximum margin that can be added on each side is 50% of the subject dimensions or 500 pixels.\n",
+            "type": "string"
+          },
+          "format": {
+            "default": "auto",
+            "description": "Result image format:\n\"auto\" = Use PNG format if transparent regions exist, otherwise use JPG format (default),\n\"png\" = PNG format with alpha transparency,\n\"jpg\" = JPG format, no transparency,\n\"zip\" = ZIP format, contains color image and alpha matte image, supports transparency (recommended).\n",
+            "enum": [
+              "auto",
+              "png",
+              "jpg",
+              "zip"
+            ],
+            "type": "string"
+          },
+          "image_file": {
+            "description": "Source image file (binary). (If this parameter is present, the other image source parameters must be empty.)",
+            "format": "binary",
+            "type": "string"
+          },
+          "image_file_b64": {
+            "description": "Source image file (base64-encoded string). (If this parameter is present, the other image source parameters must be empty.)",
+            "example": "",
+            "type": "string"
+          },
+          "image_url": {
+            "description": "Source image URL. (If this parameter is present, the other image source parameters must be empty.)",
+            "example": "https://www.remove.bg/example-hd.jpg",
+            "type": "string"
+          },
+          "position": {
+            "default": "original",
+            "description": "Positions the subject within the image canvas. Can be \"original\" (default unless \"scale\" is given), \"center\" (default when \"scale\" is given) or a value from \"0%\" to \"100%\" (both horizontal and vertical) or two values (horizontal, vertical).\n",
+            "type": "string"
+          },
+          "roi": {
+            "default": "0% 0% 100% 100%",
+            "description": "Region of interest: Only contents of this rectangular region can be detected as foreground. Everything outside is considered background and will be removed. The rectangle is defined as two x/y coordinates in the format \"x1 y1 x2 y2\". The coordinates can be in absolute pixels (suffix 'px') or relative to the width/height of the image (suffix '%'). By default, the whole image is the region of interest (\"0% 0% 100% 100%\").\n",
+            "type": "string"
+          },
+          "scale": {
+            "default": "original",
+            "description": "Scales the subject relative to the total image size. Can be any value from \"10%\" to \"100%\", or \"original\" (default). Scaling the subject implies \"position=center\" (unless specified otherwise).\n",
+            "type": "string"
+          },
+          "semitransparency": {
+            "default": "true",
+            "description": "Whether to have semi-transparent regions in the result (default: true). NOTE: Semitransparency is currently only supported for car windows (this might change in the future). Other objects are returned without semitransparency, even if set to true.\n",
+            "type": "boolean"
+          },
+          "size": {
+            "default": "preview",
+            "description": "Maximum output image resolution:\n\"preview\" (default) = Resize image to 0.25 megapixels (e.g. 625×400 pixels) – 0.25 credits per image,\n\"full\" = Use original image resolution, up to 25 megapixels (e.g. 6250x4000) with formats ZIP or JPG, or up to 10 megapixels (e.g. 4000x2500) with PNG – 1 credit per image),\n\"auto\" = Use highest available resolution (based on image size and available credits).\n\nFor backwards-compatibility this parameter also accepts the values \"medium\" (up to 1.5 megapixels) and \"hd\" (up to 4 megapixels) for 1 credit per image. The value \"full\" is also available under the name \"4k\" and the value \"preview\" is aliased as \"small\" and \"regular\".\n",
+            "enum": [
+              "preview",
+              "full",
+              "auto"
+            ],
+            "type": "string"
+          },
+          "type": {
+            "default": "auto",
+            "description": "Foreground type:\n\"auto\" = Automatically detect kind of foreground,\n\"person\" = Use person(s) as foreground,\n\"product\" = Use product(s) as foreground.\n\"car\" = Use car as foreground,\n",
+            "enum": [
+              "auto",
+              "person",
+              "product",
+              "car"
+            ],
+            "type": "string"
+          },
+          "type_level": {
+            "default": 1,
+            "description": "Classification level of the detected foreground type:\n\"none\" = No classification (X-Type Header won't bet set on the response)\n\"1\" = Use coarse classification classes: [person, product, animal, car, other]\n\"2\" = Use more specific classification classes: [person, product, animal, car, car_interior, car_part, transportation, graphics, other]\n\"latest\" = Always use the latest classification classes available\n",
+            "enum": [
+              "none",
+              1,
+              2,
+              "latest"
+            ],
+            "type": "string"
+          }
+        },
+        "type": "object"
+      }
+    },
+    "securitySchemes": {
+      "APIKeyHeader": {
+        "in": "header",
+        "name": "X-API-Key",
+        "type": "apiKey"
+      }
+    }
+  }
+}
\ No newline at end of file
