WebAPI missing files pt2

diff --git a/Engine/Plugins/Web/WebAPI/Source/WebAPIOpenAPI/Private/Tests/Samples/V3/tfl_openapi.json b/Engine/Plugins/Web/WebAPI/Source/WebAPIOpenAPI/Private/Tests/Samples/V3/tfl_openapi.json
new file mode 100644
--- /dev/null
+++ b/Engine/Plugins/Web/WebAPI/Source/WebAPIOpenAPI/Private/Tests/Samples/V3/tfl_openapi.json
@@ -0,0 +1,9939 @@
+{
+  "openapi": "3.0.0",
+  "servers": [
+    {
+      "url": "https://api.digital.tfl.gov.uk"
+    }
+  ],
+  "info": {
+    "description": "Our unified API brings together data across all modes of transport into a single RESTful API. This API provides access to the most highly requested realtime and status infomation across all the modes of transport, in a single and consistent way. Access to the developer documentation is available at https://api.tfl.gov.uk",
+    "title": "Transport for London Unified API",
+    "version": "v1",
+    "x-apisguru-categories": [
+      "transport",
+      "open_data"
+    ],
+    "x-logo": {
+      "url": "https://api.apis.guru/v2/cache/logo/https_twitter.com_TfL_profile_image.jpeg"
+    },
+    "x-origin": [
+      {
+        "format": "openapi",
+        "url": "https://api.tfl.gov.uk/swagger/docs/v1",
+        "version": "3.0"
+      }
+    ],
+    "x-providerName": "tfl.gov.uk"
+  },
+  "paths": {
+    "/AccidentStats/{year}": {
+      "get": {
+        "deprecated": false,
+        "operationId": "AccidentStats_Get",
+        "parameters": [
+          {
+            "description": "The year for which to filter the accidents on.",
+            "in": "path",
+            "name": "year",
+            "required": true,
+            "schema": {
+              "format": "int32",
+              "type": "integer"
+            }
+          }
+        ],
+        "responses": {
+          "200": {
+            "content": {
+              "application/json": {
+                "schema": {
+                  "items": {
+                    "$ref": "#/components/schemas/Tfl.Api.Presentation.Entities.AccidentStats.AccidentDetail"
+                  },
+                  "type": "array"
+                }
+              },
+              "application/xml": {
+                "schema": {
+                  "items": {
+                    "$ref": "#/components/schemas/Tfl.Api.Presentation.Entities.AccidentStats.AccidentDetail"
+                  },
+                  "type": "array"
+                }
+              },
+              "text/json": {
+                "schema": {
+                  "items": {
+                    "$ref": "#/components/schemas/Tfl.Api.Presentation.Entities.AccidentStats.AccidentDetail"
+                  },
+                  "type": "array"
+                }
+              },
+              "text/xml": {
+                "schema": {
+                  "items": {
+                    "$ref": "#/components/schemas/Tfl.Api.Presentation.Entities.AccidentStats.AccidentDetail"
+                  },
+                  "type": "array"
+                }
+              }
+            },
+            "description": "OK"
+          }
+        },
+        "summary": "Gets all accident details for accidents occuring in the specified year",
+        "tags": [
+          "AccidentStats"
+        ]
+      }
+    },
+    "/AirQuality": {
+      "get": {
+        "deprecated": false,
+        "operationId": "AirQuality_Get",
+        "responses": {
+          "200": {
+            "content": {
+              "application/json": {
+                "schema": {
+                  "$ref": "#/components/schemas/System.Object"
+                }
+              },
+              "application/xml": {
+                "schema": {
+                  "$ref": "#/components/schemas/System.Object"
+                }
+              },
+              "text/json": {
+                "schema": {
+                  "$ref": "#/components/schemas/System.Object"
+                }
+              },
+              "text/xml": {
+                "schema": {
+                  "$ref": "#/components/schemas/System.Object"
+                }
+              }
+            },
+            "description": "OK"
+          }
+        },
+        "summary": "Gets air quality data feed",
+        "tags": [
+          "AirQuality"
+        ]
+      }
+    },
+    "/BikePoint": {
+      "get": {
+        "deprecated": false,
+        "operationId": "BikePoint_GetAll",
+        "responses": {
+          "200": {
+            "content": {
+              "application/json": {
+                "schema": {
+                  "items": {
+                    "$ref": "#/components/schemas/Tfl.Api.Presentation.Entities.Place"
+                  },
+                  "type": "array"
+                }
+              },
+              "application/xml": {
+                "schema": {
+                  "items": {
+                    "$ref": "#/components/schemas/Tfl.Api.Presentation.Entities.Place"
+                  },
+                  "type": "array"
+                }
+              },
+              "text/json": {
+                "schema": {
+                  "items": {
+                    "$ref": "#/components/schemas/Tfl.Api.Presentation.Entities.Place"
+                  },
+                  "type": "array"
+                }
+              },
+              "text/xml": {
+                "schema": {
+                  "items": {
+                    "$ref": "#/components/schemas/Tfl.Api.Presentation.Entities.Place"
+                  },
+                  "type": "array"
+                }
+              }
+            },
+            "description": "OK"
+          }
+        },
+        "summary": "Gets all bike point locations. The Place object has an addtionalProperties array which contains the nbBikes, nbDocks and nbSpaces\r\n            numbers which give the status of the BikePoint. A mismatch in these numbers i.e. nbDocks - (nbBikes + nbSpaces) != 0 indicates broken docks.",
+        "tags": [
+          "BikePoint"
+        ]
+      }
+    },
+    "/BikePoint/Search": {
+      "get": {
+        "deprecated": false,
+        "operationId": "BikePoint_Search",
+        "parameters": [
+          {
+            "description": "The search term e.g. \"St. James\"",
+            "in": "query",
+            "name": "query",
+            "required": true,
+            "schema": {
+              "type": "string"
+            }
+          }
+        ],
+        "responses": {
+          "200": {
+            "content": {
+              "application/json": {
+                "schema": {
+                  "items": {
+                    "$ref": "#/components/schemas/Tfl.Api.Presentation.Entities.Place"
+                  },
+                  "type": "array"
+                }
+              },
+              "application/xml": {
+                "schema": {
+                  "items": {
+                    "$ref": "#/components/schemas/Tfl.Api.Presentation.Entities.Place"
+                  },
+                  "type": "array"
+                }
+              },
+              "text/json": {
+                "schema": {
+                  "items": {
+                    "$ref": "#/components/schemas/Tfl.Api.Presentation.Entities.Place"
+                  },
+                  "type": "array"
+                }
+              },
+              "text/xml": {
+                "schema": {
+                  "items": {
+                    "$ref": "#/components/schemas/Tfl.Api.Presentation.Entities.Place"
+                  },
+                  "type": "array"
+                }
+              }
+            },
+            "description": "OK"
+          }
+        },
+        "summary": "Search for bike stations by their name, a bike point's name often contains information about the name of the street\r\n            or nearby landmarks, for example. Note that the search result does not contain the PlaceProperties i.e. the status\r\n            or occupancy of the BikePoint, to get that information you should retrieve the BikePoint by its id on /BikePoint/id.",
+        "tags": [
+          "BikePoint"
+        ]
+      }
+    },
+    "/BikePoint/{id}": {
+      "get": {
+        "deprecated": false,
+        "operationId": "BikePoint_Get",
+        "parameters": [
+          {
+            "description": "A bike point id (a list of ids can be obtained from the above BikePoint call)",
+            "in": "path",
+            "name": "id",
+            "required": true,
+            "schema": {
+              "type": "string"
+            }
+          }
+        ],
+        "responses": {
+          "200": {
+            "content": {
+              "application/json": {
+                "schema": {
+                  "$ref": "#/components/schemas/Tfl.Api.Presentation.Entities.Place"
+                }
+              },
+              "application/xml": {
+                "schema": {
+                  "$ref": "#/components/schemas/Tfl.Api.Presentation.Entities.Place"
+                }
+              },
+              "text/json": {
+                "schema": {
+                  "$ref": "#/components/schemas/Tfl.Api.Presentation.Entities.Place"
+                }
+              },
+              "text/xml": {
+                "schema": {
+                  "$ref": "#/components/schemas/Tfl.Api.Presentation.Entities.Place"
+                }
+              }
+            },
+            "description": "OK"
+          }
+        },
+        "summary": "Gets the bike point with the given id.",
+        "tags": [
+          "BikePoint"
+        ]
+      }
+    },
+    "/Cabwise/search": {
+      "get": {
+        "deprecated": false,
+        "operationId": "Cabwise_Get",
+        "parameters": [
+          {
+            "description": "Latitude",
+            "in": "query",
+            "name": "lat",
+            "required": true,
+            "schema": {
+              "format": "double",
+              "type": "number"
+            }
+          },
+          {
+            "description": "Longitude",
+            "in": "query",
+            "name": "lon",
+            "required": true,
+            "schema": {
+              "format": "double",
+              "type": "number"
+            }
+          },
+          {
+            "description": "Operator Type e.g Minicab, Executive, Limousine",
+            "in": "query",
+            "name": "optype",
+            "required": false,
+            "schema": {
+              "type": "string"
+            }
+          },
+          {
+            "description": "Wheelchair accessible",
+            "in": "query",
+            "name": "wc",
+            "required": false,
+            "schema": {
+              "type": "string"
+            }
+          },
+          {
+            "description": "The radius of the bounding circle in metres",
+            "in": "query",
+            "name": "radius",
+            "required": false,
+            "schema": {
+              "format": "double",
+              "type": "number"
+            }
+          },
+          {
+            "description": "Trading name of operating company",
+            "in": "query",
+            "name": "name",
+            "required": false,
+            "schema": {
+              "type": "string"
+            }
+          },
+          {
+            "description": "An optional parameter to limit the number of results return. Default and maximum is 20.",
+            "in": "query",
+            "name": "maxResults",
+            "required": false,
+            "schema": {
+              "format": "int32",
+              "type": "integer"
+            }
+          },
+          {
+            "description": "Legacy Format",
+            "in": "query",
+            "name": "legacyFormat",
+            "required": false,
+            "schema": {
+              "type": "boolean"
+            }
+          },
+          {
+            "description": "Force Xml",
+            "in": "query",
+            "name": "forceXml",
+            "required": false,
+            "schema": {
+              "type": "boolean"
+            }
+          },
+          {
+            "description": "Twenty Four Seven Only",
+            "in": "query",
+            "name": "twentyFourSevenOnly",
+            "required": false,
+            "schema": {
+              "type": "boolean"
+            }
+          }
+        ],
+        "responses": {
+          "200": {
+            "content": {
+              "application/json": {
+                "schema": {
+                  "$ref": "#/components/schemas/System.Object"
+                }
+              },
+              "application/xml": {
+                "schema": {
+                  "$ref": "#/components/schemas/System.Object"
+                }
+              },
+              "text/json": {
+                "schema": {
+                  "$ref": "#/components/schemas/System.Object"
+                }
+              },
+              "text/xml": {
+                "schema": {
+                  "$ref": "#/components/schemas/System.Object"
+                }
+              }
+            },
+            "description": "OK"
+          }
+        },
+        "summary": "Gets taxis and minicabs contact information",
+        "tags": [
+          "Cabwise"
+        ]
+      }
+    },
+    "/Journey/JourneyResults/{from}/to/{to}": {
+      "get": {
+        "deprecated": false,
+        "operationId": "Journey_JourneyResults",
+        "parameters": [
+          {
+            "description": "Origin of the journey. Can be WGS84 coordinates expressed as \"lat,long\", a UK postcode, a Naptan (StopPoint) id, an ICS StopId, or a free-text string (will cause disambiguation unless it exactly matches a point of interest name).",
+            "in": "path",
+            "name": "from",
+            "required": true,
+            "schema": {
+              "type": "string"
+            }
+          },
+          {
+            "description": "Destination of the journey. Can be WGS84 coordinates expressed as \"lat,long\", a UK postcode, a Naptan (StopPoint) id, an ICS StopId, or a free-text string (will cause disambiguation unless it exactly matches a point of interest name).",
+            "in": "path",
+            "name": "to",
+            "required": true,
+            "schema": {
+              "type": "string"
+            }
+          },
+          {
+            "description": "Travel through point on the journey. Can be WGS84 coordinates expressed as \"lat,long\", a UK postcode, a Naptan (StopPoint) id, an ICS StopId, or a free-text string (will cause disambiguation unless it exactly matches a point of interest name).",
+            "in": "query",
+            "name": "via",
+            "required": false,
+            "schema": {
+              "type": "string"
+            }
+          },
+          {
+            "description": "Does the journey cover stops outside London? eg. \"nationalSearch=true\"",
+            "in": "query",
+            "name": "nationalSearch",
+            "required": false,
+            "schema": {
+              "type": "boolean"
+            }
+          },
+          {
+            "description": "The date must be in yyyyMMdd format",
+            "in": "query",
+            "name": "date",
+            "required": false,
+            "schema": {
+              "type": "string"
+            }
+          },
+          {
+            "description": "The time must be in HHmm format",
+            "in": "query",
+            "name": "time",
+            "required": false,
+            "schema": {
+              "type": "string"
+            }
+          },
+          {
+            "description": "Does the time given relate to arrival or leaving time? Possible options: \"departing\" | \"arriving\"",
+            "in": "query",
+            "name": "timeIs",
+            "required": false,
+            "schema": {
+              "enum": [
+                "Arriving",
+                "Departing"
+              ],
+              "type": "string"
+            }
+          },
+          {
+            "description": "The journey preference eg possible options: \"leastinterchange\" | \"leasttime\" | \"leastwalking\"",
+            "in": "query",
+            "name": "journeyPreference",
+            "required": false,
+            "schema": {
+              "enum": [
+                "LeastInterchange",
+                "LeastTime",
+                "LeastWalking"
+              ],
+              "type": "string"
+            }
+          },
+          {
+            "description": "The mode must be a comma separated list of modes. eg possible options: \"public-bus,overground,train,tube,coach,dlr,cablecar,tram,river,walking,cycle\"",
+            "explode": true,
+            "in": "query",
+            "name": "mode",
+            "required": false,
+            "schema": {
+              "items": {
+                "type": "string"
+              },
+              "type": "array"
+            }
+          },
+          {
+            "description": "The accessibility preference must be a comma separated list eg. \"noSolidStairs,noEscalators,noElevators,stepFreeToVehicle,stepFreeToPlatform\"",
+            "explode": true,
+            "in": "query",
+            "name": "accessibilityPreference",
+            "required": false,
+            "schema": {
+              "items": {
+                "enum": [
+                  "NoRequirements",
+                  "NoSolidStairs",
+                  "NoEscalators",
+                  "NoElevators",
+                  "StepFreeToVehicle",
+                  "StepFreeToPlatform"
+                ],
+                "type": "string"
+              },
+              "type": "array"
+            }
+          },
+          {
+            "description": "An optional name to associate with the origin of the journey in the results.",
+            "in": "query",
+            "name": "fromName",
+            "required": false,
+            "schema": {
+              "type": "string"
+            }
+          },
+          {
+            "description": "An optional name to associate with the destination of the journey in the results.",
+            "in": "query",
+            "name": "toName",
+            "required": false,
+            "schema": {
+              "type": "string"
+            }
+          },
+          {
+            "description": "An optional name to associate with the via point of the journey in the results.",
+            "in": "query",
+            "name": "viaName",
+            "required": false,
+            "schema": {
+              "type": "string"
+            }
+          },
+          {
+            "description": "The max walking time in minutes for transfer eg. \"120\"",
+            "in": "query",
+            "name": "maxTransferMinutes",
+            "required": false,
+            "schema": {
+              "type": "string"
+            }
+          },
+          {
+            "description": "The max walking time in minutes for journeys eg. \"120\"",
+            "in": "query",
+            "name": "maxWalkingMinutes",
+            "required": false,
+            "schema": {
+              "type": "string"
+            }
+          },
+          {
+            "description": "The walking speed. eg possible options: \"slow\" | \"average\" | \"fast\".",
+            "in": "query",
+            "name": "walkingSpeed",
+            "required": false,
+            "schema": {
+              "enum": [
+                "Slow",
+                "Average",
+                "Fast"
+              ],
+              "type": "string"
+            }
+          },
+          {
+            "description": "The cycle preference. eg possible options: \"allTheWay\" | \"leaveAtStation\" | \"takeOnTransport\" | \"cycleHire\"",
+            "in": "query",
+            "name": "cyclePreference",
+            "required": false,
+            "schema": {
+              "enum": [
+                "None",
+                "LeaveAtStation",
+                "TakeOnTransport",
+                "AllTheWay",
+                "CycleHire"
+              ],
+              "type": "string"
+            }
+          },
+          {
+            "description": "Time adjustment command. eg possible options: \"TripFirst\" | \"TripLast\"",
+            "in": "query",
+            "name": "adjustment",
+            "required": false,
+            "schema": {
+              "type": "string"
+            }
+          },
+          {
+            "description": "A comma separated list of cycling proficiency levels. eg possible options: \"easy,moderate,fast\"",
+            "explode": true,
+            "in": "query",
+            "name": "bikeProficiency",
+            "required": false,
+            "schema": {
+              "items": {
+                "enum": [
+                  "Easy",
+                  "Moderate",
+                  "Fast"
+                ],
+                "type": "string"
+              },
+              "type": "array"
+            }
+          },
+          {
+            "description": "Option to determine whether to return alternative cycling journey",
+            "in": "query",
+            "name": "alternativeCycle",
+            "required": false,
+            "schema": {
+              "type": "boolean"
+            }
+          },
+          {
+            "description": "Option to determine whether to return alternative walking journey",
+            "in": "query",
+            "name": "alternativeWalking",
+            "required": false,
+            "schema": {
+              "type": "boolean"
+            }
+          },
+          {
+            "description": "Flag to determine whether certain text (e.g. walking instructions) should be output with HTML tags or not.",
+            "in": "query",
+            "name": "applyHtmlMarkup",
+            "required": false,
+            "schema": {
+              "type": "boolean"
+            }
+          },
+          {
+            "description": "A boolean to indicate whether or not to return 3 public transport journeys, a bus journey, a cycle hire journey, a personal cycle journey and a walking journey",
+            "in": "query",
+            "name": "useMultiModalCall",
+            "required": false,
+            "schema": {
+              "type": "boolean"
+            }
+          },
+          {
+            "description": "A boolean to indicate whether to optimize journeys using walking",
+            "in": "query",
+            "name": "walkingOptimization",
+            "required": false,
+            "schema": {
+              "type": "boolean"
+            }
+          },
+          {
+            "description": "A boolean to indicate whether to return one or more taxi journeys. Note, setting this to true will override \"useMultiModalCall\".",
+            "in": "query",
+            "name": "taxiOnlyTrip",
+            "required": false,
+            "schema": {
+              "type": "boolean"
+            }
+          },
+          {
+            "description": "A boolean to indicate whether public transport routes should include directions between platforms and station entrances.",
+            "in": "query",
+            "name": "routeBetweenEntrances",
+            "required": false,
+            "schema": {
+              "type": "boolean"
+            }
+          },
+          {
+            "description": "A boolean to indicate if we want to receive real time live arrivals data where available.",
+            "in": "query",
+            "name": "useRealTimeLiveArrivals",
+            "required": false,
+            "schema": {
+              "type": "boolean"
+            }
+          }
+        ],
+        "responses": {
+          "200": {
+            "content": {
+              "application/json": {
+                "schema": {
+                  "$ref": "#/components/schemas/Tfl.Api.Presentation.Entities.JourneyPlanner.ItineraryResult"
+                }
+              },
+              "application/xml": {
+                "schema": {
+                  "$ref": "#/components/schemas/Tfl.Api.Presentation.Entities.JourneyPlanner.ItineraryResult"
+                }
+              },
+              "text/json": {
+                "schema": {
+                  "$ref": "#/components/schemas/Tfl.Api.Presentation.Entities.JourneyPlanner.ItineraryResult"
+                }
+              },
+              "text/xml": {
+                "schema": {
+                  "$ref": "#/components/schemas/Tfl.Api.Presentation.Entities.JourneyPlanner.ItineraryResult"
+                }
+              }
+            },
+            "description": "OK"
+          }
+        },
+        "summary": "Perform a Journey Planner search from the parameters specified in simple types",
+        "tags": [
+          "Journey"
+        ]
+      }
+    },
+    "/Journey/Meta/Modes": {
+      "get": {
+        "deprecated": false,
+        "operationId": "Journey_Meta",
+        "responses": {
+          "200": {
+            "content": {
+              "application/json": {
+                "schema": {
+                  "items": {
+                    "$ref": "#/components/schemas/Tfl.Api.Presentation.Entities.Mode"
+                  },
+                  "type": "array"
+                }
+              },
+              "application/xml": {
+                "schema": {
+                  "items": {
+                    "$ref": "#/components/schemas/Tfl.Api.Presentation.Entities.Mode"
+                  },
+                  "type": "array"
+                }
+              },
+              "text/json": {
+                "schema": {
+                  "items": {
+                    "$ref": "#/components/schemas/Tfl.Api.Presentation.Entities.Mode"
+                  },
+                  "type": "array"
+                }
+              },
+              "text/xml": {
+                "schema": {
+                  "items": {
+                    "$ref": "#/components/schemas/Tfl.Api.Presentation.Entities.Mode"
+                  },
+                  "type": "array"
+                }
+              }
+            },
+            "description": "OK"
+          }
+        },
+        "summary": "Gets a list of all of the available journey planner modes",
+        "tags": [
+          "Journey"
+        ]
+      }
+    },
+    "/Line/Meta/DisruptionCategories": {
+      "get": {
+        "deprecated": false,
+        "operationId": "Line_MetaDisruptionCategories",
+        "responses": {
+          "200": {
+            "content": {
+              "application/json": {
+                "schema": {
+                  "items": {
+                    "type": "string"
+                  },
+                  "type": "array"
+                }
+              },
+              "application/xml": {
+                "schema": {
+                  "items": {
+                    "type": "string"
+                  },
+                  "type": "array"
+                }
+              },
+              "text/json": {
+                "schema": {
+                  "items": {
+                    "type": "string"
+                  },
+                  "type": "array"
+                }
+              },
+              "text/xml": {
+                "schema": {
+                  "items": {
+                    "type": "string"
+                  },
+                  "type": "array"
+                }
+              }
+            },
+            "description": "OK"
+          }
+        },
+        "summary": "Gets a list of valid disruption categories",
+        "tags": [
+          "Line"
+        ]
+      }
+    },
+    "/Line/Meta/Modes": {
+      "get": {
+        "deprecated": false,
+        "operationId": "Line_MetaModes",
+        "responses": {
+          "200": {
+            "content": {
+              "application/json": {
+                "schema": {
+                  "items": {
+                    "$ref": "#/components/schemas/Tfl.Api.Presentation.Entities.Mode"
+                  },
+                  "type": "array"
+                }
+              },
+              "application/xml": {
+                "schema": {
+                  "items": {
+                    "$ref": "#/components/schemas/Tfl.Api.Presentation.Entities.Mode"
+                  },
+                  "type": "array"
+                }
+              },
+              "text/json": {
+                "schema": {
+                  "items": {
+                    "$ref": "#/components/schemas/Tfl.Api.Presentation.Entities.Mode"
+                  },
+                  "type": "array"
+                }
+              },
+              "text/xml": {
+                "schema": {
+                  "items": {
+                    "$ref": "#/components/schemas/Tfl.Api.Presentation.Entities.Mode"
+                  },
+                  "type": "array"
+                }
+              }
+            },
+            "description": "OK"
+          }
+        },
+        "summary": "Gets a list of valid modes",
+        "tags": [
+          "Line"
+        ]
+      }
+    },
+    "/Line/Meta/ServiceTypes": {
+      "get": {
+        "deprecated": false,
+        "operationId": "Line_MetaServiceTypes",
+        "responses": {
+          "200": {
+            "content": {
+              "application/json": {
+                "schema": {
+                  "items": {
+                    "type": "string"
+                  },
+                  "type": "array"
+                }
+              },
+              "application/xml": {
+                "schema": {
+                  "items": {
+                    "type": "string"
+                  },
+                  "type": "array"
+                }
+              },
+              "text/json": {
+                "schema": {
+                  "items": {
+                    "type": "string"
+                  },
+                  "type": "array"
+                }
+              },
+              "text/xml": {
+                "schema": {
+                  "items": {
+                    "type": "string"
+                  },
+                  "type": "array"
+                }
+              }
+            },
+            "description": "OK"
+          }
+        },
+        "summary": "Gets a list of valid ServiceTypes to filter on",
+        "tags": [
+          "Line"
+        ]
+      }
+    },
+    "/Line/Meta/Severity": {
+      "get": {
+        "deprecated": false,
+        "operationId": "Line_MetaSeverity",
+        "responses": {
+          "200": {
+            "content": {
+              "application/json": {
+                "schema": {
+                  "items": {
+                    "$ref": "#/components/schemas/Tfl.Api.Presentation.Entities.StatusSeverity"
+                  },
+                  "type": "array"
+                }
+              },
+              "application/xml": {
+                "schema": {
+                  "items": {
+                    "$ref": "#/components/schemas/Tfl.Api.Presentation.Entities.StatusSeverity"
+                  },
+                  "type": "array"
+                }
+              },
+              "text/json": {
+                "schema": {
+                  "items": {
+                    "$ref": "#/components/schemas/Tfl.Api.Presentation.Entities.StatusSeverity"
+                  },
+                  "type": "array"
+                }
+              },
+              "text/xml": {
+                "schema": {
+                  "items": {
+                    "$ref": "#/components/schemas/Tfl.Api.Presentation.Entities.StatusSeverity"
+                  },
+                  "type": "array"
+                }
+              }
+            },
+            "description": "OK"
+          }
+        },
+        "summary": "Gets a list of valid severity codes",
+        "tags": [
+          "Line"
+        ]
+      }
+    },
+    "/Line/Mode/{modes}": {
+      "get": {
+        "deprecated": false,
+        "operationId": "Line_GetByMode",
+        "parameters": [
+          {
+            "description": "A comma-separated list of modes e.g. tube,dlr",
+            "in": "path",
+            "name": "modes",
+            "required": true,
+            "schema": {
+              "items": {
+                "type": "string"
+              },
+              "type": "array"
+            },
+            "style": "simple"
+          }
+        ],
+        "responses": {
+          "200": {
+            "content": {
+              "application/json": {
+                "schema": {
+                  "items": {
+                    "$ref": "#/components/schemas/Tfl.Api.Presentation.Entities.Line"
+                  },
+                  "type": "array"
+                }
+              },
+              "application/xml": {
+                "schema": {
+                  "items": {
+                    "$ref": "#/components/schemas/Tfl.Api.Presentation.Entities.Line"
+                  },
+                  "type": "array"
+                }
+              },
+              "text/json": {
+                "schema": {
+                  "items": {
+                    "$ref": "#/components/schemas/Tfl.Api.Presentation.Entities.Line"
+                  },
+                  "type": "array"
+                }
+              },
+              "text/xml": {
+                "schema": {
+                  "items": {
+                    "$ref": "#/components/schemas/Tfl.Api.Presentation.Entities.Line"
+                  },
+                  "type": "array"
+                }
+              }
+            },
+            "description": "OK"
+          }
+        },
+        "summary": "Gets lines that serve the given modes.",
+        "tags": [
+          "Line"
+        ]
+      }
+    },
+    "/Line/Mode/{modes}/Disruption": {
+      "get": {
+        "deprecated": false,
+        "operationId": "Line_DisruptionByMode",
+        "parameters": [
+          {
+            "description": "A comma-separated list of modes e.g. tube,dlr",
+            "in": "path",
+            "name": "modes",
+            "required": true,
+            "schema": {
+              "items": {
+                "type": "string"
+              },
+              "type": "array"
+            },
+            "style": "simple"
+          }
+        ],
+        "responses": {
+          "200": {
+            "content": {
+              "application/json": {
+                "schema": {
+                  "items": {
+                    "$ref": "#/components/schemas/Tfl.Api.Presentation.Entities.Disruption"
+                  },
+                  "type": "array"
+                }
+              },
+              "application/xml": {
+                "schema": {
+                  "items": {
+                    "$ref": "#/components/schemas/Tfl.Api.Presentation.Entities.Disruption"
+                  },
+                  "type": "array"
+                }
+              },
+              "text/json": {
+                "schema": {
+                  "items": {
+                    "$ref": "#/components/schemas/Tfl.Api.Presentation.Entities.Disruption"
+                  },
+                  "type": "array"
+                }
+              },
+              "text/xml": {
+                "schema": {
+                  "items": {
+                    "$ref": "#/components/schemas/Tfl.Api.Presentation.Entities.Disruption"
+                  },
+                  "type": "array"
+                }
+              }
+            },
+            "description": "OK"
+          }
+        },
+        "summary": "Get disruptions for all lines of the given modes.",
+        "tags": [
+          "Line"
+        ]
+      }
+    },
+    "/Line/Mode/{modes}/Route": {
+      "get": {
+        "deprecated": false,
+        "operationId": "Line_RouteByMode",
+        "parameters": [
+          {
+            "description": "A comma-separated list of modes e.g. tube,dlr",
+            "in": "path",
+            "name": "modes",
+            "required": true,
+            "schema": {
+              "items": {
+                "type": "string"
+              },
+              "type": "array"
+            },
+            "style": "simple"
+          },
+          {
+            "description": "A comma seperated list of service types to filter on. Supported values: Regular, Night. Defaulted to 'Regular' if not specified",
+            "explode": true,
+            "in": "query",
+            "name": "serviceTypes",
+            "required": false,
+            "schema": {
+              "items": {
+                "enum": [
+                  "Regular",
+                  "Night"
+                ],
+                "type": "string"
+              },
+              "type": "array"
+            }
+          }
+        ],
+        "responses": {
+          "200": {
+            "content": {
+              "application/json": {
+                "schema": {
+                  "items": {
+                    "$ref": "#/components/schemas/Tfl.Api.Presentation.Entities.Line"
+                  },
+                  "type": "array"
+                }
+              },
+              "application/xml": {
+                "schema": {
+                  "items": {
+                    "$ref": "#/components/schemas/Tfl.Api.Presentation.Entities.Line"
+                  },
+                  "type": "array"
+                }
+              },
+              "text/json": {
+                "schema": {
+                  "items": {
+                    "$ref": "#/components/schemas/Tfl.Api.Presentation.Entities.Line"
+                  },
+                  "type": "array"
+                }
+              },
+              "text/xml": {
+                "schema": {
+                  "items": {
+                    "$ref": "#/components/schemas/Tfl.Api.Presentation.Entities.Line"
+                  },
+                  "type": "array"
+                }
+              }
+            },
+            "description": "OK"
+          }
+        },
+        "summary": "Gets all lines and their valid routes for given modes, including the name and id of the originating and terminating stops for each route",
+        "tags": [
+          "Line"
+        ]
+      }
+    },
+    "/Line/Mode/{modes}/Status": {
+      "get": {
+        "deprecated": false,
+        "operationId": "Line_StatusByMode",
+        "parameters": [
+          {
+            "description": "A comma-separated list of modes to filter by. e.g. tube,dlr",
+            "in": "path",
+            "name": "modes",
+            "required": true,
+            "schema": {
+              "items": {
+                "type": "string"
+              },
+              "type": "array"
+            },
+            "style": "simple"
+          },
+          {
+            "description": "Include details of the disruptions that are causing the line status including the affected stops and routes",
+            "in": "query",
+            "name": "detail",
+            "required": false,
+            "schema": {
+              "type": "boolean"
+            }
+          },
+          {
+            "description": "If specified, ensures that only those line status(es) are returned within the lines that have disruptions with the matching severity level.",
+            "in": "query",
+            "name": "severityLevel",
+            "required": false,
+            "schema": {
+              "type": "string"
+            }
+          }
+        ],
+        "responses": {
+          "200": {
+            "content": {
+              "application/json": {
+                "schema": {
+                  "items": {
+                    "$ref": "#/components/schemas/Tfl.Api.Presentation.Entities.Line"
+                  },
+                  "type": "array"
+                }
+              },
+              "application/xml": {
+                "schema": {
+                  "items": {
+                    "$ref": "#/components/schemas/Tfl.Api.Presentation.Entities.Line"
+                  },
+                  "type": "array"
+                }
+              },
+              "text/json": {
+                "schema": {
+                  "items": {
+                    "$ref": "#/components/schemas/Tfl.Api.Presentation.Entities.Line"
+                  },
+                  "type": "array"
+                }
+              },
+              "text/xml": {
+                "schema": {
+                  "items": {
+                    "$ref": "#/components/schemas/Tfl.Api.Presentation.Entities.Line"
+                  },
+                  "type": "array"
+                }
+              }
+            },
+            "description": "OK"
+          }
+        },
+        "summary": "Gets the line status of for all lines for the given modes",
+        "tags": [
+          "Line"
+        ]
+      }
+    },
+    "/Line/Route": {
+      "get": {
+        "deprecated": false,
+        "operationId": "Line_Route",
+        "parameters": [
+          {
+            "description": "A comma seperated list of service types to filter on. Supported values: Regular, Night. Defaulted to 'Regular' if not specified",
+            "explode": true,
+            "in": "query",
+            "name": "serviceTypes",
+            "required": false,
+            "schema": {
+              "items": {
+                "enum": [
+                  "Regular",
+                  "Night"
+                ],
+                "type": "string"
+              },
+              "type": "array"
+            }
+          }
+        ],
+        "responses": {
+          "200": {
+            "content": {
+              "application/json": {
+                "schema": {
+                  "items": {
+                    "$ref": "#/components/schemas/Tfl.Api.Presentation.Entities.Line"
+                  },
+                  "type": "array"
+                }
+              },
+              "application/xml": {
+                "schema": {
+                  "items": {
+                    "$ref": "#/components/schemas/Tfl.Api.Presentation.Entities.Line"
+                  },
+                  "type": "array"
+                }
+              },
+              "text/json": {
+                "schema": {
+                  "items": {
+                    "$ref": "#/components/schemas/Tfl.Api.Presentation.Entities.Line"
+                  },
+                  "type": "array"
+                }
+              },
+              "text/xml": {
+                "schema": {
+                  "items": {
+                    "$ref": "#/components/schemas/Tfl.Api.Presentation.Entities.Line"
+                  },
+                  "type": "array"
+                }
+              }
+            },
+            "description": "OK"
+          }
+        },
+        "summary": "Get all valid routes for all lines, including the name and id of the originating and terminating stops for each route.",
+        "tags": [
+          "Line"
+        ]
+      }
+    },
+    "/Line/Search/{query}": {
+      "get": {
+        "deprecated": false,
+        "operationId": "Line_Search",
+        "parameters": [
+          {
+            "description": "Search term e.g victoria",
+            "in": "path",
+            "name": "query",
+            "required": true,
+            "schema": {
+              "type": "string"
+            }
+          },
+          {
+            "description": "Optionally filter by the specified modes",
+            "explode": true,
+            "in": "query",
+            "name": "modes",
+            "required": false,
+            "schema": {
+              "items": {
+                "type": "string"
+              },
+              "type": "array"
+            }
+          },
+          {
+            "description": "A comma seperated list of service types to filter on. Supported values: Regular, Night. Defaulted to 'Regular' if not specified",
+            "explode": true,
+            "in": "query",
+            "name": "serviceTypes",
+            "required": false,
+            "schema": {
+              "items": {
+                "enum": [
+                  "Regular",
+                  "Night"
+                ],
+                "type": "string"
+              },
+              "type": "array"
+            }
+          }
+        ],
+        "responses": {
+          "200": {
+            "content": {
+              "application/json": {
+                "schema": {
+                  "$ref": "#/components/schemas/Tfl.Api.Presentation.Entities.RouteSearchResponse"
+                }
+              },
+              "application/xml": {
+                "schema": {
+                  "$ref": "#/components/schemas/Tfl.Api.Presentation.Entities.RouteSearchResponse"
+                }
+              },
+              "text/json": {
+                "schema": {
+                  "$ref": "#/components/schemas/Tfl.Api.Presentation.Entities.RouteSearchResponse"
+                }
+              },
+              "text/xml": {
+                "schema": {
+                  "$ref": "#/components/schemas/Tfl.Api.Presentation.Entities.RouteSearchResponse"
+                }
+              }
+            },
+            "description": "OK"
+          }
+        },
+        "summary": "Search for lines or routes matching the query string",
+        "tags": [
+          "Line"
+        ]
+      }
+    },
+    "/Line/Status/{severity}": {
+      "get": {
+        "deprecated": false,
+        "operationId": "Line_StatusBySeverity",
+        "parameters": [
+          {
+            "description": "The level of severity (eg: a number from 0 to 14)",
+            "in": "path",
+            "name": "severity",
+            "required": true,
+            "schema": {
+              "format": "int32",
+              "type": "integer"
+            }
+          }
+        ],
+        "responses": {
+          "200": {
+            "content": {
+              "application/json": {
+                "schema": {
+                  "items": {
+                    "$ref": "#/components/schemas/Tfl.Api.Presentation.Entities.Line"
+                  },
+                  "type": "array"
+                }
+              },
+              "application/xml": {
+                "schema": {
+                  "items": {
+                    "$ref": "#/components/schemas/Tfl.Api.Presentation.Entities.Line"
+                  },
+                  "type": "array"
+                }
+              },
+              "text/json": {
+                "schema": {
+                  "items": {
+                    "$ref": "#/components/schemas/Tfl.Api.Presentation.Entities.Line"
+                  },
+                  "type": "array"
+                }
+              },
+              "text/xml": {
+                "schema": {
+                  "items": {
+                    "$ref": "#/components/schemas/Tfl.Api.Presentation.Entities.Line"
+                  },
+                  "type": "array"
+                }
+              }
+            },
+            "description": "OK"
+          }
+        },
+        "summary": "Gets the line status for all lines with a given severity\r\n            A list of valid severity codes can be obtained from a call to Line/Meta/Severity",
+        "tags": [
+          "Line"
+        ]
+      }
+    },
+    "/Line/{ids}": {
+      "get": {
+        "deprecated": false,
+        "operationId": "Line_Get",
+        "parameters": [
+          {
+            "description": "A comma-separated list of line ids e.g. victoria,circle,N133. Max. approx. 20 ids.",
+            "in": "path",
+            "name": "ids",
+            "required": true,
+            "schema": {
+              "items": {
+                "type": "string"
+              },
+              "type": "array"
+            },
+            "style": "simple"
+          }
+        ],
+        "responses": {
+          "200": {
+            "content": {
+              "application/json": {
+                "schema": {
+                  "items": {
+                    "$ref": "#/components/schemas/Tfl.Api.Presentation.Entities.Line"
+                  },
+                  "type": "array"
+                }
+              },
+              "application/xml": {
+                "schema": {
+                  "items": {
+                    "$ref": "#/components/schemas/Tfl.Api.Presentation.Entities.Line"
+                  },
+                  "type": "array"
+                }
+              },
+              "text/json": {
+                "schema": {
+                  "items": {
+                    "$ref": "#/components/schemas/Tfl.Api.Presentation.Entities.Line"
+                  },
+                  "type": "array"
+                }
+              },
+              "text/xml": {
+                "schema": {
+                  "items": {
+                    "$ref": "#/components/schemas/Tfl.Api.Presentation.Entities.Line"
+                  },
+                  "type": "array"
+                }
+              }
+            },
+            "description": "OK"
+          }
+        },
+        "summary": "Gets lines that match the specified line ids.",
+        "tags": [
+          "Line"
+        ]
+      }
+    },
+    "/Line/{ids}/Arrivals/{stopPointId}": {
+      "get": {
+        "deprecated": false,
+        "operationId": "Line_Arrivals",
+        "parameters": [
+          {
+            "description": "A comma-separated list of line ids e.g. victoria,circle,N133. Max. approx. 20 ids.",
+            "in": "path",
+            "name": "ids",
+            "required": true,
+            "schema": {
+              "items": {
+                "type": "string"
+              },
+              "type": "array"
+            },
+            "style": "simple"
+          },
+          {
+            "description": "Optional. Id of stop to get arrival predictions for (station naptan code e.g. 940GZZLUASL, you can use /StopPoint/Search/{query} endpoint to find a stop point id from a station name)",
+            "in": "path",
+            "name": "stopPointId",
+            "required": true,
+            "schema": {
+              "type": "string"
+            }
+          },
+          {
+            "description": "Optional. The direction of travel. Can be inbound or outbound or all. If left blank, and destinationStopId is set, will default to all",
+            "in": "query",
+            "name": "direction",
+            "required": false,
+            "schema": {
+              "enum": [
+                "inbound",
+                "outbound",
+                "all"
+              ],
+              "type": "string"
+            }
+          },
+          {
+            "description": "Optional. Id of destination stop",
+            "in": "query",
+            "name": "destinationStationId",
+            "required": false,
+            "schema": {
+              "type": "string"
+            }
+          }
+        ],
+        "responses": {
+          "200": {
+            "content": {
+              "application/json": {
+                "schema": {
+                  "items": {
+                    "$ref": "#/components/schemas/Tfl.Api.Presentation.Entities.Prediction"
+                  },
+                  "type": "array"
+                }
+              },
+              "application/xml": {
+                "schema": {
+                  "items": {
+                    "$ref": "#/components/schemas/Tfl.Api.Presentation.Entities.Prediction"
+                  },
+                  "type": "array"
+                }
+              },
+              "text/json": {
+                "schema": {
+                  "items": {
+                    "$ref": "#/components/schemas/Tfl.Api.Presentation.Entities.Prediction"
+                  },
+                  "type": "array"
+                }
+              },
+              "text/xml": {
+                "schema": {
+                  "items": {
+                    "$ref": "#/components/schemas/Tfl.Api.Presentation.Entities.Prediction"
+                  },
+                  "type": "array"
+                }
+              }
+            },
+            "description": "OK"
+          }
+        },
+        "summary": "Get the list of arrival predictions for given line ids based at the given stop",
+        "tags": [
+          "Line"
+        ]
+      }
+    },
+    "/Line/{ids}/Disruption": {
+      "get": {
+        "deprecated": false,
+        "operationId": "Line_Disruption",
+        "parameters": [
+          {
+            "description": "A comma-separated list of line ids e.g. victoria,circle,N133. Max. approx. 20 ids.",
+            "in": "path",
+            "name": "ids",
+            "required": true,
+            "schema": {
+              "items": {
+                "type": "string"
+              },
+              "type": "array"
+            },
+            "style": "simple"
+          }
+        ],
+        "responses": {
+          "200": {
+            "content": {
+              "application/json": {
+                "schema": {
+                  "items": {
+                    "$ref": "#/components/schemas/Tfl.Api.Presentation.Entities.Disruption"
+                  },
+                  "type": "array"
+                }
+              },
+              "application/xml": {
+                "schema": {
+                  "items": {
+                    "$ref": "#/components/schemas/Tfl.Api.Presentation.Entities.Disruption"
+                  },
+                  "type": "array"
+                }
+              },
+              "text/json": {
+                "schema": {
+                  "items": {
+                    "$ref": "#/components/schemas/Tfl.Api.Presentation.Entities.Disruption"
+                  },
+                  "type": "array"
+                }
+              },
+              "text/xml": {
+                "schema": {
+                  "items": {
+                    "$ref": "#/components/schemas/Tfl.Api.Presentation.Entities.Disruption"
+                  },
+                  "type": "array"
+                }
+              }
+            },
+            "description": "OK"
+          }
+        },
+        "summary": "Get disruptions for the given line ids",
+        "tags": [
+          "Line"
+        ]
+      }
+    },
+    "/Line/{ids}/Route": {
+      "get": {
+        "deprecated": false,
+        "operationId": "Line_LineRoutesByIds",
+        "parameters": [
+          {
+            "description": "A comma-separated list of line ids e.g. victoria,circle,N133. Max. approx. 20 ids.",
+            "in": "path",
+            "name": "ids",
+            "required": true,
+            "schema": {
+              "items": {
+                "type": "string"
+              },
+              "type": "array"
+            },
+            "style": "simple"
+          },
+          {
+            "description": "A comma seperated list of service types to filter on. Supported values: Regular, Night. Defaulted to 'Regular' if not specified",
+            "explode": true,
+            "in": "query",
+            "name": "serviceTypes",
+            "required": false,
+            "schema": {
+              "items": {
+                "enum": [
+                  "Regular",
+                  "Night"
+                ],
+                "type": "string"
+              },
+              "type": "array"
+            }
+          }
+        ],
+        "responses": {
+          "200": {
+            "content": {
+              "application/json": {
+                "schema": {
+                  "items": {
+                    "$ref": "#/components/schemas/Tfl.Api.Presentation.Entities.Line"
+                  },
+                  "type": "array"
+                }
+              },
+              "application/xml": {
+                "schema": {
+                  "items": {
+                    "$ref": "#/components/schemas/Tfl.Api.Presentation.Entities.Line"
+                  },
+                  "type": "array"
+                }
+              },
+              "text/json": {
+                "schema": {
+                  "items": {
+                    "$ref": "#/components/schemas/Tfl.Api.Presentation.Entities.Line"
+                  },
+                  "type": "array"
+                }
+              },
+              "text/xml": {
+                "schema": {
+                  "items": {
+                    "$ref": "#/components/schemas/Tfl.Api.Presentation.Entities.Line"
+                  },
+                  "type": "array"
+                }
+              }
+            },
+            "description": "OK"
+          }
+        },
+        "summary": "Get all valid routes for given line ids, including the name and id of the originating and terminating stops for each route.",
+        "tags": [
+          "Line"
+        ]
+      }
+    },
+    "/Line/{ids}/Status": {
+      "get": {
+        "deprecated": false,
+        "operationId": "Line_StatusByIds",
+        "parameters": [
+          {
+            "description": "A comma-separated list of line ids e.g. victoria,circle,N133. Max. approx. 20 ids.",
+            "in": "path",
+            "name": "ids",
+            "required": true,
+            "schema": {
+              "items": {
+                "type": "string"
+              },
+              "type": "array"
+            },
+            "style": "simple"
+          },
+          {
+            "description": "Include details of the disruptions that are causing the line status including the affected stops and routes",
+            "in": "query",
+            "name": "detail",
+            "required": false,
+            "schema": {
+              "type": "boolean"
+            }
+          }
+        ],
+        "responses": {
+          "200": {
+            "content": {
+              "application/json": {
+                "schema": {
+                  "items": {
+                    "$ref": "#/components/schemas/Tfl.Api.Presentation.Entities.Line"
+                  },
+                  "type": "array"
+                }
+              },
+              "application/xml": {
+                "schema": {
+                  "items": {
+                    "$ref": "#/components/schemas/Tfl.Api.Presentation.Entities.Line"
+                  },
+                  "type": "array"
+                }
+              },
+              "text/json": {
+                "schema": {
+                  "items": {
+                    "$ref": "#/components/schemas/Tfl.Api.Presentation.Entities.Line"
+                  },
+                  "type": "array"
+                }
+              },
+              "text/xml": {
+                "schema": {
+                  "items": {
+                    "$ref": "#/components/schemas/Tfl.Api.Presentation.Entities.Line"
+                  },
+                  "type": "array"
+                }
+              }
+            },
+            "description": "OK"
+          }
+        },
+        "summary": "Gets the line status of for given line ids e.g Minor Delays",
+        "tags": [
+          "Line"
+        ]
+      }
+    },
+    "/Line/{ids}/Status/{StartDate}/to/{EndDate}": {
+      "get": {
+        "deprecated": false,
+        "operationId": "Line_Status",
+        "parameters": [
+          {
+            "description": "A comma-separated list of line ids e.g. victoria,circle,N133. Max. approx. 20 ids.",
+            "in": "path",
+            "name": "ids",
+            "required": true,
+            "schema": {
+              "items": {
+                "type": "string"
+              },
+              "type": "array"
+            },
+            "style": "simple"
+          },
+          {
+            "description": "Include details of the disruptions that are causing the line status including the affected stops and routes",
+            "in": "query",
+            "name": "detail",
+            "required": false,
+            "schema": {
+              "type": "boolean"
+            }
+          },
+          {
+            "in": "query",
+            "name": "startDate",
+            "required": true,
+            "schema": {
+              "type": "string"
+            }
+          },
+          {
+            "in": "query",
+            "name": "endDate",
+            "required": true,
+            "schema": {
+              "type": "string"
+            }
+          },
+          {
+            "in": "query",
+            "name": "dateRange.startDate",
+            "required": false,
+            "schema": {
+              "format": "date-time",
+              "type": "string"
+            }
+          },
+          {
+            "in": "query",
+            "name": "dateRange.endDate",
+            "required": false,
+            "schema": {
+              "format": "date-time",
+              "type": "string"
+            }
+          },
+          {
+            "description": "Automatically added",
+            "in": "path",
+            "name": "StartDate",
+            "required": true,
+            "schema": {
+              "type": "string"
+            }
+          },
+          {
+            "description": "Automatically added",
+            "in": "path",
+            "name": "EndDate",
+            "required": true,
+            "schema": {
+              "type": "string"
+            }
+          }
+        ],
+        "responses": {
+          "200": {
+            "content": {
+              "application/json": {
+                "schema": {
+                  "items": {
+                    "$ref": "#/components/schemas/Tfl.Api.Presentation.Entities.Line"
+                  },
+                  "type": "array"
+                }
+              },
+              "application/xml": {
+                "schema": {
+                  "items": {
+                    "$ref": "#/components/schemas/Tfl.Api.Presentation.Entities.Line"
+                  },
+                  "type": "array"
+                }
+              },
+              "text/json": {
+                "schema": {
+                  "items": {
+                    "$ref": "#/components/schemas/Tfl.Api.Presentation.Entities.Line"
+                  },
+                  "type": "array"
+                }
+              },
+              "text/xml": {
+                "schema": {
+                  "items": {
+                    "$ref": "#/components/schemas/Tfl.Api.Presentation.Entities.Line"
+                  },
+                  "type": "array"
+                }
+              }
+            },
+            "description": "OK"
+          }
+        },
+        "summary": "Gets the line status for given line ids during the provided dates e.g Minor Delays",
+        "tags": [
+          "Line"
+        ]
+      }
+    },
+    "/Line/{id}/Route/Sequence/{direction}": {
+      "get": {
+        "deprecated": false,
+        "operationId": "Line_RouteSequence",
+        "parameters": [
+          {
+            "description": "A single line id e.g. victoria",
+            "in": "path",
+            "name": "id",
+            "required": true,
+            "schema": {
+              "type": "string"
+            }
+          },
+          {
+            "description": "The direction of travel. Can be inbound or outbound.",
+            "in": "path",
+            "name": "direction",
+            "required": true,
+            "schema": {
+              "enum": [
+                "inbound",
+                "outbound",
+                "all"
+              ],
+              "type": "string"
+            }
+          },
+          {
+            "description": "A comma seperated list of service types to filter on. Supported values: Regular, Night. Defaulted to 'Regular' if not specified",
+            "explode": true,
+            "in": "query",
+            "name": "serviceTypes",
+            "required": false,
+            "schema": {
+              "items": {
+                "enum": [
+                  "Regular",
+                  "Night"
+                ],
+                "type": "string"
+              },
+              "type": "array"
+            }
+          },
+          {
+            "description": "That excludes crowding from line disruptions. Can be true or false.",
+            "in": "query",
+            "name": "excludeCrowding",
+            "required": false,
+            "schema": {
+              "type": "boolean"
+            }
+          }
+        ],
+        "responses": {
+          "200": {
+            "content": {
+              "application/json": {
+                "schema": {
+                  "$ref": "#/components/schemas/Tfl.Api.Presentation.Entities.RouteSequence"
+                }
+              },
+              "application/xml": {
+                "schema": {
+                  "$ref": "#/components/schemas/Tfl.Api.Presentation.Entities.RouteSequence"
+                }
+              },
+              "text/json": {
+                "schema": {
+                  "$ref": "#/components/schemas/Tfl.Api.Presentation.Entities.RouteSequence"
+                }
+              },
+              "text/xml": {
+                "schema": {
+                  "$ref": "#/components/schemas/Tfl.Api.Presentation.Entities.RouteSequence"
+                }
+              }
+            },
+            "description": "OK"
+          }
+        },
+        "summary": "Gets all valid routes for given line id, including the sequence of stops on each route.",
+        "tags": [
+          "Line"
+        ]
+      }
+    },
+    "/Line/{id}/StopPoints": {
+      "get": {
+        "deprecated": false,
+        "operationId": "Line_StopPoints",
+        "parameters": [
+          {
+            "description": "A single line id e.g. victoria",
+            "in": "path",
+            "name": "id",
+            "required": true,
+            "schema": {
+              "type": "string"
+            }
+          },
+          {
+            "description": "If the national-rail line is requested, this flag will filter the national rail stations so that only those operated by TfL are returned",
+            "in": "query",
+            "name": "tflOperatedNationalRailStationsOnly",
+            "required": false,
+            "schema": {
+              "type": "boolean"
+            }
+          }
+        ],
+        "responses": {
+          "200": {
+            "content": {
+              "application/json": {
+                "schema": {
+                  "items": {
+                    "$ref": "#/components/schemas/Tfl.Api.Presentation.Entities.StopPoint"
+                  },
+                  "type": "array"
+                }
+              },
+              "application/xml": {
+                "schema": {
+                  "items": {
+                    "$ref": "#/components/schemas/Tfl.Api.Presentation.Entities.StopPoint"
+                  },
+                  "type": "array"
+                }
+              },
+              "text/json": {
+                "schema": {
+                  "items": {
+                    "$ref": "#/components/schemas/Tfl.Api.Presentation.Entities.StopPoint"
+                  },
+                  "type": "array"
+                }
+              },
+              "text/xml": {
+                "schema": {
+                  "items": {
+                    "$ref": "#/components/schemas/Tfl.Api.Presentation.Entities.StopPoint"
+                  },
+                  "type": "array"
+                }
+              }
+            },
+            "description": "OK"
+          }
+        },
+        "summary": "Gets a list of the stations that serve the given line id",
+        "tags": [
+          "Line"
+        ]
+      }
+    },
+    "/Line/{id}/Timetable/{fromStopPointId}": {
+      "get": {
+        "deprecated": false,
+        "operationId": "Line_Timetable",
+        "parameters": [
+          {
+            "description": "The originating station's stop point id (station naptan code e.g. 940GZZLUASL, you can use /StopPoint/Search/{query} endpoint to find a stop point id from a station name)",
+            "in": "path",
+            "name": "fromStopPointId",
+            "required": true,
+            "schema": {
+              "type": "string"
+            }
+          },
+          {
+            "description": "A single line id e.g. victoria",
+            "in": "path",
+            "name": "id",
+            "required": true,
+            "schema": {
+              "type": "string"
+            }
+          }
+        ],
+        "responses": {
+          "200": {
+            "content": {
+              "application/json": {
+                "schema": {
+                  "$ref": "#/components/schemas/Tfl.Api.Presentation.Entities.TimetableResponse"
+                }
+              },
+              "application/xml": {
+                "schema": {
+                  "$ref": "#/components/schemas/Tfl.Api.Presentation.Entities.TimetableResponse"
+                }
+              },
+              "text/json": {
+                "schema": {
+                  "$ref": "#/components/schemas/Tfl.Api.Presentation.Entities.TimetableResponse"
+                }
+              },
+              "text/xml": {
+                "schema": {
+                  "$ref": "#/components/schemas/Tfl.Api.Presentation.Entities.TimetableResponse"
+                }
+              }
+            },
+            "description": "OK"
+          }
+        },
+        "summary": "Gets the timetable for a specified station on the give line",
+        "tags": [
+          "Line"
+        ]
+      }
+    },
+    "/Line/{id}/Timetable/{fromStopPointId}/to/{toStopPointId}": {
+      "get": {
+        "deprecated": false,
+        "operationId": "Line_TimetableTo",
+        "parameters": [
+          {
+            "description": "The originating station's stop point id (station naptan code e.g. 940GZZLUASL, you can use /StopPoint/Search/{query} endpoint to find a stop point id from a station name)",
+            "in": "path",
+            "name": "fromStopPointId",
+            "required": true,
+            "schema": {
+              "type": "string"
+            }
+          },
+          {
+            "description": "A single line id e.g. victoria",
+            "in": "path",
+            "name": "id",
+            "required": true,
+            "schema": {
+              "type": "string"
+            }
+          },
+          {
+            "description": "The destination stations's Naptan code",
+            "in": "path",
+            "name": "toStopPointId",
+            "required": true,
+            "schema": {
+              "type": "string"
+            }
+          }
+        ],
+        "responses": {
+          "200": {
+            "content": {
+              "application/json": {
+                "schema": {
+                  "$ref": "#/components/schemas/Tfl.Api.Presentation.Entities.TimetableResponse"
+                }
+              },
+              "application/xml": {
+                "schema": {
+                  "$ref": "#/components/schemas/Tfl.Api.Presentation.Entities.TimetableResponse"
+                }
+              },
+              "text/json": {
+                "schema": {
+                  "$ref": "#/components/schemas/Tfl.Api.Presentation.Entities.TimetableResponse"
+                }
+              },
+              "text/xml": {
+                "schema": {
+                  "$ref": "#/components/schemas/Tfl.Api.Presentation.Entities.TimetableResponse"
+                }
+              }
+            },
+            "description": "OK"
+          }
+        },
+        "summary": "Gets the timetable for a specified station on the give line with specified destination",
+        "tags": [
+          "Line"
+        ]
+      }
+    },
+    "/Mode/ActiveServiceTypes": {
+      "get": {
+        "deprecated": false,
+        "operationId": "Mode_GetActiveServiceTypes",
+        "responses": {
+          "200": {
+            "content": {
+              "application/json": {
+                "schema": {
+                  "items": {
+                    "$ref": "#/components/schemas/Tfl.Api.Presentation.Entities.ActiveServiceType"
+                  },
+                  "type": "array"
+                }
+              },
+              "application/xml": {
+                "schema": {
+                  "items": {
+                    "$ref": "#/components/schemas/Tfl.Api.Presentation.Entities.ActiveServiceType"
+                  },
+                  "type": "array"
+                }
+              },
+              "text/json": {
+                "schema": {
+                  "items": {
+                    "$ref": "#/components/schemas/Tfl.Api.Presentation.Entities.ActiveServiceType"
+                  },
+                  "type": "array"
+                }
+              },
+              "text/xml": {
+                "schema": {
+                  "items": {
+                    "$ref": "#/components/schemas/Tfl.Api.Presentation.Entities.ActiveServiceType"
+                  },
+                  "type": "array"
+                }
+              }
+            },
+            "description": "OK"
+          }
+        },
+        "summary": "Returns the service type active for a mode.\r\n            Currently only supports tube",
+        "tags": [
+          "Mode"
+        ]
+      }
+    },
+    "/Mode/{mode}/Arrivals": {
+      "get": {
+        "deprecated": false,
+        "operationId": "Mode_Arrivals",
+        "parameters": [
+          {
+            "description": "A mode name e.g. tube, dlr",
+            "in": "path",
+            "name": "mode",
+            "required": true,
+            "schema": {
+              "type": "string"
+            }
+          },
+          {
+            "description": "A number of arrivals to return for each stop, -1 to return all available.",
+            "in": "query",
+            "name": "count",
+            "required": false,
+            "schema": {
+              "format": "int32",
+              "type": "integer"
+            }
+          }
+        ],
+        "responses": {
+          "200": {
+            "content": {
+              "application/json": {
+                "schema": {
+                  "items": {
+                    "$ref": "#/components/schemas/Tfl.Api.Presentation.Entities.Prediction"
+                  },
+                  "type": "array"
+                }
+              },
+              "application/xml": {
+                "schema": {
+                  "items": {
+                    "$ref": "#/components/schemas/Tfl.Api.Presentation.Entities.Prediction"
+                  },
+                  "type": "array"
+                }
+              },
+              "text/json": {
+                "schema": {
+                  "items": {
+                    "$ref": "#/components/schemas/Tfl.Api.Presentation.Entities.Prediction"
+                  },
+                  "type": "array"
+                }
+              },
+              "text/xml": {
+                "schema": {
+                  "items": {
+                    "$ref": "#/components/schemas/Tfl.Api.Presentation.Entities.Prediction"
+                  },
+                  "type": "array"
+                }
+              }
+            },
+            "description": "OK"
+          }
+        },
+        "summary": "Gets the next arrival predictions for all stops of a given mode",
+        "tags": [
+          "Mode"
+        ]
+      }
+    },
+    "/Occupancy/BikePoints/{ids}": {
+      "get": {
+        "deprecated": false,
+        "operationId": "Occupancy_GetBikePointsOccupancies",
+        "parameters": [
+          {
+            "description": "",
+            "in": "path",
+            "name": "ids",
+            "required": true,
+            "schema": {
+              "items": {
+                "type": "string"
+              },
+              "type": "array"
+            },
+            "style": "simple"
+          }
+        ],
+        "responses": {
+          "200": {
+            "content": {
+              "application/json": {
+                "schema": {
+                  "items": {
+                    "$ref": "#/components/schemas/Tfl.Api.Presentation.Entities.BikePointOccupancy"
+                  },
+                  "type": "array"
+                }
+              },
+              "application/xml": {
+                "schema": {
+                  "items": {
+                    "$ref": "#/components/schemas/Tfl.Api.Presentation.Entities.BikePointOccupancy"
+                  },
+                  "type": "array"
+                }
+              },
+              "text/json": {
+                "schema": {
+                  "items": {
+                    "$ref": "#/components/schemas/Tfl.Api.Presentation.Entities.BikePointOccupancy"
+                  },
+                  "type": "array"
+                }
+              },
+              "text/xml": {
+                "schema": {
+                  "items": {
+                    "$ref": "#/components/schemas/Tfl.Api.Presentation.Entities.BikePointOccupancy"
+                  },
+                  "type": "array"
+                }
+              }
+            },
+            "description": "OK"
+          }
+        },
+        "summary": "Get the occupancy for bike points.",
+        "tags": [
+          "Occupancy"
+        ]
+      }
+    },
+    "/Occupancy/CarPark": {
+      "get": {
+        "deprecated": false,
+        "responses": {
+          "200": {
+            "content": {
+              "application/json": {
+                "schema": {
+                  "items": {
+                    "$ref": "#/components/schemas/Tfl.Api.Presentation.Entities.CarParkOccupancy"
+                  },
+                  "type": "array"
+                }
+              },
+              "application/xml": {
+                "schema": {
+                  "items": {
+                    "$ref": "#/components/schemas/Tfl.Api.Presentation.Entities.CarParkOccupancy"
+                  },
+                  "type": "array"
+                }
+              },
+              "text/json": {
+                "schema": {
+                  "items": {
+                    "$ref": "#/components/schemas/Tfl.Api.Presentation.Entities.CarParkOccupancy"
+                  },
+                  "type": "array"
+                }
+              },
+              "text/xml": {
+                "schema": {
+                  "items": {
+                    "$ref": "#/components/schemas/Tfl.Api.Presentation.Entities.CarParkOccupancy"
+                  },
+                  "type": "array"
+                }
+              }
+            },
+            "description": "OK"
+          }
+        },
+        "summary": "Gets the occupancy for all car parks that have occupancy data",
+        "tags": [
+          "Occupancy"
+        ]
+      }
+    },
+    "/Occupancy/CarPark/{id}": {
+      "get": {
+        "deprecated": false,
+        "operationId": "Occupancy_Get",
+        "parameters": [
+          {
+            "description": "",
+            "in": "path",
+            "name": "id",
+            "required": true,
+            "schema": {
+              "type": "string"
+            }
+          }
+        ],
+        "responses": {
+          "200": {
+            "content": {
+              "application/json": {
+                "schema": {
+                  "$ref": "#/components/schemas/Tfl.Api.Presentation.Entities.CarParkOccupancy"
+                }
+              },
+              "application/xml": {
+                "schema": {
+                  "$ref": "#/components/schemas/Tfl.Api.Presentation.Entities.CarParkOccupancy"
+                }
+              },
+              "text/json": {
+                "schema": {
+                  "$ref": "#/components/schemas/Tfl.Api.Presentation.Entities.CarParkOccupancy"
+                }
+              },
+              "text/xml": {
+                "schema": {
+                  "$ref": "#/components/schemas/Tfl.Api.Presentation.Entities.CarParkOccupancy"
+                }
+              }
+            },
+            "description": "OK"
+          }
+        },
+        "summary": "Gets the occupancy for a car park with a given id",
+        "tags": [
+          "Occupancy"
+        ]
+      }
+    },
+    "/Occupancy/ChargeConnector": {
+      "get": {
+        "deprecated": false,
+        "operationId": "Occupancy_GetAllChargeConnectorStatus",
+        "responses": {
+          "200": {
+            "content": {
+              "application/json": {
+                "schema": {
+                  "items": {
+                    "$ref": "#/components/schemas/Tfl.Api.Presentation.Entities.ChargeConnectorOccupancy"
+                  },
+                  "type": "array"
+                }
+              },
+              "application/xml": {
+                "schema": {
+                  "items": {
+                    "$ref": "#/components/schemas/Tfl.Api.Presentation.Entities.ChargeConnectorOccupancy"
+                  },
+                  "type": "array"
+                }
+              },
+              "text/json": {
+                "schema": {
+                  "items": {
+                    "$ref": "#/components/schemas/Tfl.Api.Presentation.Entities.ChargeConnectorOccupancy"
+                  },
+                  "type": "array"
+                }
+              },
+              "text/xml": {
+                "schema": {
+                  "items": {
+                    "$ref": "#/components/schemas/Tfl.Api.Presentation.Entities.ChargeConnectorOccupancy"
+                  },
+                  "type": "array"
+                }
+              }
+            },
+            "description": "OK"
+          }
+        },
+        "summary": "Gets the occupancy for all charge connectors",
+        "tags": [
+          "Occupancy"
+        ]
+      }
+    },
+    "/Occupancy/ChargeConnector/{ids}": {
+      "get": {
+        "deprecated": false,
+        "operationId": "Occupancy_GetChargeConnectorStatus",
+        "parameters": [
+          {
+            "in": "path",
+            "name": "ids",
+            "required": true,
+            "schema": {
+              "items": {
+                "type": "string"
+              },
+              "type": "array"
+            },
+            "style": "simple"
+          }
+        ],
+        "responses": {
+          "200": {
+            "content": {
+              "application/json": {
+                "schema": {
+                  "items": {
+                    "$ref": "#/components/schemas/Tfl.Api.Presentation.Entities.ChargeConnectorOccupancy"
+                  },
+                  "type": "array"
+                }
+              },
+              "application/xml": {
+                "schema": {
+                  "items": {
+                    "$ref": "#/components/schemas/Tfl.Api.Presentation.Entities.ChargeConnectorOccupancy"
+                  },
+                  "type": "array"
+                }
+              },
+              "text/json": {
+                "schema": {
+                  "items": {
+                    "$ref": "#/components/schemas/Tfl.Api.Presentation.Entities.ChargeConnectorOccupancy"
+                  },
+                  "type": "array"
+                }
+              },
+              "text/xml": {
+                "schema": {
+                  "items": {
+                    "$ref": "#/components/schemas/Tfl.Api.Presentation.Entities.ChargeConnectorOccupancy"
+                  },
+                  "type": "array"
+                }
+              }
+            },
+            "description": "OK"
+          }
+        },
+        "summary": "Gets the occupancy for a charge connectors with a given id (sourceSystemPlaceId)",
+        "tags": [
+          "Occupancy"
+        ]
+      }
+    },
+    "/Place": {
+      "get": {
+        "deprecated": false,
+        "operationId": "Place_GetByGeo",
+        "parameters": [
+          {
+            "description": "The radius of the bounding circle in metres when only lat/lon are specified.",
+            "in": "query",
+            "name": "radius",
+            "required": false,
+            "schema": {
+              "format": "double",
+              "type": "number"
+            }
+          },
+          {
+            "description": "An optional list of comma separated property categories to return in the Place's property bag. If null or empty, all categories of property are returned. Pass the keyword \"none\" to return no properties (a valid list of categories can be obtained from the /Place/Meta/categories endpoint)",
+            "explode": true,
+            "in": "query",
+            "name": "categories",
+            "required": false,
+            "schema": {
+              "items": {
+                "type": "string"
+              },
+              "type": "array"
+            }
+          },
+          {
+            "description": "Defaults to false. If true child places e.g. individual charging stations at a charge point while be included, otherwise just the URLs of any child places will be returned",
+            "in": "query",
+            "name": "includeChildren",
+            "required": false,
+            "schema": {
+              "type": "boolean"
+            }
+          },
+          {
+            "description": "Place types to filter on, or null to return all types",
+            "explode": true,
+            "in": "query",
+            "name": "type",
+            "required": false,
+            "schema": {
+              "items": {
+                "type": "string"
+              },
+              "type": "array"
+            }
+          },
+          {
+            "description": "An optional parameter to limit the results to active records only (Currently only the 'VariableMessageSign' place type is supported)",
+            "in": "query",
+            "name": "activeOnly",
+            "required": false,
+            "schema": {
+              "type": "boolean"
+            }
+          },
+          {
+            "description": "If specified, limits the number of returned places equal to the given value",
+            "in": "query",
+            "name": "numberOfPlacesToReturn",
+            "required": false,
+            "schema": {
+              "format": "int32",
+              "type": "integer"
+            }
+          },
+          {
+            "in": "query",
+            "name": "placeGeo.swLat",
+            "required": false,
+            "schema": {
+              "format": "double",
+              "type": "number"
+            }
+          },
+          {
+            "in": "query",
+            "name": "placeGeo.swLon",
+            "required": false,
+            "schema": {
+              "format": "double",
+              "type": "number"
+            }
+          },
+          {
+            "in": "query",
+            "name": "placeGeo.neLat",
+            "required": false,
+            "schema": {
+              "format": "double",
+              "type": "number"
+            }
+          },
+          {
+            "in": "query",
+            "name": "placeGeo.neLon",
+            "required": false,
+            "schema": {
+              "format": "double",
+              "type": "number"
+            }
+          },
+          {
+            "in": "query",
+            "name": "placeGeo.lat",
+            "required": false,
+            "schema": {
+              "format": "double",
+              "type": "number"
+            }
+          },
+          {
+            "in": "query",
+            "name": "placeGeo.lon",
+            "required": false,
+            "schema": {
+              "format": "double",
+              "type": "number"
+            }
+          }
+        ],
+        "responses": {
+          "200": {
+            "content": {
+              "application/json": {
+                "schema": {
+                  "items": {
+                    "$ref": "#/components/schemas/Tfl.Api.Presentation.Entities.StopPoint"
+                  },
+                  "type": "array"
+                }
+              },
+              "application/xml": {
+                "schema": {
+                  "items": {
+                    "$ref": "#/components/schemas/Tfl.Api.Presentation.Entities.StopPoint"
+                  },
+                  "type": "array"
+                }
+              },
+              "text/json": {
+                "schema": {
+                  "items": {
+                    "$ref": "#/components/schemas/Tfl.Api.Presentation.Entities.StopPoint"
+                  },
+                  "type": "array"
+                }
+              },
+              "text/xml": {
+                "schema": {
+                  "items": {
+                    "$ref": "#/components/schemas/Tfl.Api.Presentation.Entities.StopPoint"
+                  },
+                  "type": "array"
+                }
+              }
+            },
+            "description": "OK"
+          }
+        },
+        "summary": "Gets the places that lie within a geographic region. The geographic region of interest can either be specified\r\n            by using a lat/lon geo-point and a radius in metres to return places within the locus defined by the lat/lon of\r\n            its centre or alternatively, by the use of a bounding box defined by the lat/lon of its north-west and south-east corners.\r\n            Optionally filters on type and can strip properties for a smaller payload.",
+        "tags": [
+          "Place"
+        ]
+      }
+    },
+    "/Place/Address/Streets/{Postcode}": {
+      "get": {
+        "deprecated": false,
+        "operationId": "Place_GetStreetsByPostCode",
+        "parameters": [
+          {
+            "in": "query",
+            "name": "postcode",
+            "required": true,
+            "schema": {
+              "type": "string"
+            }
+          },
+          {
+            "in": "query",
+            "name": "postcodeInput.postcode",
+            "required": false,
+            "schema": {
+              "pattern": "^(([gG][iI][rR] {0,}0[aA]{2})|((([a-pr-uwyzA-PR-UWYZ][a-hk-yA-HK-Y]?[0-9][0-9]?)|(([a-pr-uwyzA-PR-UWYZ][0-9][a-hjkstuwA-HJKSTUW])|([a-pr-uwyzA-PR-UWYZ][a-hk-yA-HK-Y][0-9][abehmnprv-yABEHMNPRV-Y]))) {0,}[0-9][abd-hjlnp-uw-zABD-HJLNP-UW-Z]{2}))$",
+              "type": "string"
+            }
+          },
+          {
+            "description": "Automatically added",
+            "in": "path",
+            "name": "Postcode",
+            "required": true,
+            "schema": {
+              "type": "string"
+            }
+          }
+        ],
+        "responses": {
+          "200": {
+            "content": {
+              "application/json": {
+                "schema": {
+                  "$ref": "#/components/schemas/System.Object"
+                }
+              },
+              "application/xml": {
+                "schema": {
+                  "$ref": "#/components/schemas/System.Object"
+                }
+              },
+              "text/json": {
+                "schema": {
+                  "$ref": "#/components/schemas/System.Object"
+                }
+              },
+              "text/xml": {
+                "schema": {
+                  "$ref": "#/components/schemas/System.Object"
+                }
+              }
+            },
+            "description": "OK"
+          }
+        },
+        "summary": "Gets the set of streets associated with a post code.",
+        "tags": [
+          "Place"
+        ]
+      }
+    },
+    "/Place/Meta/Categories": {
+      "get": {
+        "deprecated": false,
+        "operationId": "Place_MetaCategories",
+        "responses": {
+          "200": {
+            "content": {
+              "application/json": {
+                "schema": {
+                  "items": {
+                    "$ref": "#/components/schemas/Tfl.Api.Presentation.Entities.PlaceCategory"
+                  },
+                  "type": "array"
+                }
+              },
+              "application/xml": {
+                "schema": {
+                  "items": {
+                    "$ref": "#/components/schemas/Tfl.Api.Presentation.Entities.PlaceCategory"
+                  },
+                  "type": "array"
+                }
+              },
+              "text/json": {
+                "schema": {
+                  "items": {
+                    "$ref": "#/components/schemas/Tfl.Api.Presentation.Entities.PlaceCategory"
+                  },
+                  "type": "array"
+                }
+              },
+              "text/xml": {
+                "schema": {
+                  "items": {
+                    "$ref": "#/components/schemas/Tfl.Api.Presentation.Entities.PlaceCategory"
+                  },
+                  "type": "array"
+                }
+              }
+            },
+            "description": "OK"
+          }
+        },
+        "summary": "Gets a list of all of the available place property categories and keys.",
+        "tags": [
+          "Place"
+        ]
+      }
+    },
+    "/Place/Meta/PlaceTypes": {
+      "get": {
+        "deprecated": false,
+        "operationId": "Place_MetaPlaceTypes",
+        "responses": {
+          "200": {
+            "content": {
+              "application/json": {
+                "schema": {
+                  "items": {
+                    "$ref": "#/components/schemas/Tfl.Api.Presentation.Entities.PlaceCategory"
+                  },
+                  "type": "array"
+                }
+              },
+              "application/xml": {
+                "schema": {
+                  "items": {
+                    "$ref": "#/components/schemas/Tfl.Api.Presentation.Entities.PlaceCategory"
+                  },
+                  "type": "array"
+                }
+              },
+              "text/json": {
+                "schema": {
+                  "items": {
+                    "$ref": "#/components/schemas/Tfl.Api.Presentation.Entities.PlaceCategory"
+                  },
+                  "type": "array"
+                }
+              },
+              "text/xml": {
+                "schema": {
+                  "items": {
+                    "$ref": "#/components/schemas/Tfl.Api.Presentation.Entities.PlaceCategory"
+                  },
+                  "type": "array"
+                }
+              }
+            },
+            "description": "OK"
+          }
+        },
+        "summary": "Gets a list of the available types of Place.",
+        "tags": [
+          "Place"
+        ]
+      }
+    },
+    "/Place/Search": {
+      "get": {
+        "deprecated": false,
+        "operationId": "Place_Search",
+        "parameters": [
+          {
+            "description": "The name of the place, you can use the /Place/Types/{types} endpoint to get a list of places for a given type including their names.",
+            "in": "query",
+            "name": "name",
+            "required": true,
+            "schema": {
+              "type": "string"
+            }
+          },
+          {
+            "description": "A comma-separated list of the types to return. Max. approx 12 types.",
+            "explode": true,
+            "in": "query",
+            "name": "types",
+            "required": false,
+            "schema": {
+              "items": {
+                "type": "string"
+              },
+              "type": "array"
+            }
+          }
+        ],
+        "responses": {
+          "200": {
+            "content": {
+              "application/json": {
+                "schema": {
+                  "items": {
+                    "$ref": "#/components/schemas/Tfl.Api.Presentation.Entities.Place"
+                  },
+                  "type": "array"
+                }
+              },
+              "application/xml": {
+                "schema": {
+                  "items": {
+                    "$ref": "#/components/schemas/Tfl.Api.Presentation.Entities.Place"
+                  },
+                  "type": "array"
+                }
+              },
+              "text/json": {
+                "schema": {
+                  "items": {
+                    "$ref": "#/components/schemas/Tfl.Api.Presentation.Entities.Place"
+                  },
+                  "type": "array"
+                }
+              },
+              "text/xml": {
+                "schema": {
+                  "items": {
+                    "$ref": "#/components/schemas/Tfl.Api.Presentation.Entities.Place"
+                  },
+                  "type": "array"
+                }
+              }
+            },
+            "description": "OK"
+          }
+        },
+        "summary": "Gets all places that matches the given query",
+        "tags": [
+          "Place"
+        ]
+      }
+    },
+    "/Place/Type/{types}": {
+      "get": {
+        "deprecated": false,
+        "operationId": "Place_GetByType",
+        "parameters": [
+          {
+            "description": "A comma-separated list of the types to return. Max. approx 12 types.\r\n            A valid list of place types can be obtained from the /Place/Meta/placeTypes endpoint.",
+            "in": "path",
+            "name": "types",
+            "required": true,
+            "schema": {
+              "items": {
+                "type": "string"
+              },
+              "type": "array"
+            },
+            "style": "simple"
+          },
+          {
+            "description": "An optional parameter to limit the results to active records only (Currently only the 'VariableMessageSign' place type is supported)",
+            "in": "query",
+            "name": "activeOnly",
+            "required": false,
+            "schema": {
+              "type": "boolean"
+            }
+          }
+        ],
+        "responses": {
+          "200": {
+            "content": {
+              "application/json": {
+                "schema": {
+                  "items": {
+                    "$ref": "#/components/schemas/Tfl.Api.Presentation.Entities.Place"
+                  },
+                  "type": "array"
+                }
+              },
+              "application/xml": {
+                "schema": {
+                  "items": {
+                    "$ref": "#/components/schemas/Tfl.Api.Presentation.Entities.Place"
+                  },
+                  "type": "array"
+                }
+              },
+              "text/json": {
+                "schema": {
+                  "items": {
+                    "$ref": "#/components/schemas/Tfl.Api.Presentation.Entities.Place"
+                  },
+                  "type": "array"
+                }
+              },
+              "text/xml": {
+                "schema": {
+                  "items": {
+                    "$ref": "#/components/schemas/Tfl.Api.Presentation.Entities.Place"
+                  },
+                  "type": "array"
+                }
+              }
+            },
+            "description": "OK"
+          }
+        },
+        "summary": "Gets all places of a given type",
+        "tags": [
+          "Place"
+        ]
+      }
+    },
+    "/Place/{id}": {
+      "get": {
+        "deprecated": false,
+        "operationId": "Place_Get",
+        "parameters": [
+          {
+            "description": "The id of the place, you can use the /Place/Types/{types} endpoint to get a list of places for a given type including their ids",
+            "in": "path",
+            "name": "id",
+            "required": true,
+            "schema": {
+              "type": "string"
+            }
+          },
+          {
+            "description": "Defaults to false. If true child places e.g. individual charging stations at a charge point while be included, otherwise just the URLs of any child places will be returned",
+            "in": "query",
+            "name": "includeChildren",
+            "required": false,
+            "schema": {
+              "type": "boolean"
+            }
+          }
+        ],
+        "responses": {
+          "200": {
+            "content": {
+              "application/json": {
+                "schema": {
+                  "items": {
+                    "$ref": "#/components/schemas/Tfl.Api.Presentation.Entities.Place"
+                  },
+                  "type": "array"
+                }
+              },
+              "application/xml": {
+                "schema": {
+                  "items": {
+                    "$ref": "#/components/schemas/Tfl.Api.Presentation.Entities.Place"
+                  },
+                  "type": "array"
+                }
+              },
+              "text/json": {
+                "schema": {
+                  "items": {
+                    "$ref": "#/components/schemas/Tfl.Api.Presentation.Entities.Place"
+                  },
+                  "type": "array"
+                }
+              },
+              "text/xml": {
+                "schema": {
+                  "items": {
+                    "$ref": "#/components/schemas/Tfl.Api.Presentation.Entities.Place"
+                  },
+                  "type": "array"
+                }
+              }
+            },
+            "description": "OK"
+          }
+        },
+        "summary": "Gets the place with the given id.",
+        "tags": [
+          "Place"
+        ]
+      }
+    },
+    "/Place/{type}/At/{Lat}/{Lon}": {
+      "get": {
+        "deprecated": false,
+        "operationId": "Place_GetAt",
+        "parameters": [
+          {
+            "description": "The place type (a valid list of place types can be obtained from the /Place/Meta/placeTypes endpoint)",
+            "in": "path",
+            "name": "type",
+            "required": true,
+            "schema": {
+              "items": {
+                "type": "string"
+              },
+              "type": "array"
+            },
+            "style": "simple"
+          },
+          {
+            "in": "query",
+            "name": "lat",
+            "required": true,
+            "schema": {
+              "type": "string"
+            }
+          },
+          {
+            "in": "query",
+            "name": "lon",
+            "required": true,
+            "schema": {
+              "type": "string"
+            }
+          },
+          {
+            "in": "query",
+            "name": "location.lat",
+            "required": true,
+            "schema": {
+              "format": "double",
+              "type": "number"
+            }
+          },
+          {
+            "in": "query",
+            "name": "location.lon",
+            "required": true,
+            "schema": {
+              "format": "double",
+              "type": "number"
+            }
+          },
+          {
+            "description": "Automatically added",
+            "in": "path",
+            "name": "Lat",
+            "required": true,
+            "schema": {
+              "type": "string"
+            }
+          },
+          {
+            "description": "Automatically added",
+            "in": "path",
+            "name": "Lon",
+            "required": true,
+            "schema": {
+              "type": "string"
+            }
+          }
+        ],
+        "responses": {
+          "200": {
+            "content": {
+              "application/json": {
+                "schema": {
+                  "$ref": "#/components/schemas/System.Object"
+                }
+              },
+              "application/xml": {
+                "schema": {
+                  "$ref": "#/components/schemas/System.Object"
+                }
+              },
+              "text/json": {
+                "schema": {
+                  "$ref": "#/components/schemas/System.Object"
+                }
+              },
+              "text/xml": {
+                "schema": {
+                  "$ref": "#/components/schemas/System.Object"
+                }
+              }
+            },
+            "description": "OK"
+          }
+        },
+        "summary": "Gets any places of the given type whose geography intersects the given latitude and longitude. In practice this means the Place\r\n            must be polygonal e.g. a BoroughBoundary.",
+        "tags": [
+          "Place"
+        ]
+      }
+    },
+    "/Place/{type}/overlay/{z}/{Lat}/{Lon}/{width}/{height}": {
+      "get": {
+        "deprecated": false,
+        "operationId": "Place_GetOverlay",
+        "parameters": [
+          {
+            "description": "The zoom level",
+            "in": "path",
+            "name": "z",
+            "required": true,
+            "schema": {
+              "format": "int32",
+              "type": "integer"
+            }
+          },
+          {
+            "description": "The place type (a valid list of place types can be obtained from the /Place/Meta/placeTypes endpoint)",
+            "in": "path",
+            "name": "type",
+            "required": true,
+            "schema": {
+              "items": {
+                "type": "string"
+              },
+              "type": "array"
+            },
+            "style": "simple"
+          },
+          {
+            "description": "The width of the requested overlay.",
+            "in": "path",
+            "name": "width",
+            "required": true,
+            "schema": {
+              "format": "int32",
+              "type": "integer"
+            }
+          },
+          {
+            "description": "The height of the requested overlay.",
+            "in": "path",
+            "name": "height",
+            "required": true,
+            "schema": {
+              "format": "int32",
+              "type": "integer"
+            }
+          },
+          {
+            "in": "query",
+            "name": "lat",
+            "required": true,
+            "schema": {
+              "type": "string"
+            }
+          },
+          {
+            "in": "query",
+            "name": "lon",
+            "required": true,
+            "schema": {
+              "type": "string"
+            }
+          },
+          {
+            "in": "query",
+            "name": "location.lat",
+            "required": true,
+            "schema": {
+              "format": "double",
+              "type": "number"
+            }
+          },
+          {
+            "in": "query",
+            "name": "location.lon",
+            "required": true,
+            "schema": {
+              "format": "double",
+              "type": "number"
+            }
+          },
+          {
+            "description": "Automatically added",
+            "in": "path",
+            "name": "Lat",
+            "required": true,
+            "schema": {
+              "type": "string"
+            }
+          },
+          {
+            "description": "Automatically added",
+            "in": "path",
+            "name": "Lon",
+            "required": true,
+            "schema": {
+              "type": "string"
+            }
+          }
+        ],
+        "responses": {
+          "200": {
+            "content": {
+              "application/json": {
+                "schema": {
+                  "$ref": "#/components/schemas/System.Object"
+                }
+              },
+              "application/xml": {
+                "schema": {
+                  "$ref": "#/components/schemas/System.Object"
+                }
+              },
+              "text/json": {
+                "schema": {
+                  "$ref": "#/components/schemas/System.Object"
+                }
+              },
+              "text/xml": {
+                "schema": {
+                  "$ref": "#/components/schemas/System.Object"
+                }
+              }
+            },
+            "description": "OK"
+          }
+        },
+        "summary": "Gets the place overlay for a given set of co-ordinates and a given width/height.",
+        "tags": [
+          "Place"
+        ]
+      }
+    },
+    "/Road": {
+      "get": {
+        "deprecated": false,
+        "operationId": "Road_Get",
+        "responses": {
+          "200": {
+            "content": {
+              "application/json": {
+                "schema": {
+                  "items": {
+                    "$ref": "#/components/schemas/Tfl.Api.Presentation.Entities.RoadCorridor"
+                  },
+                  "type": "array"
+                }
+              },
+              "application/xml": {
+                "schema": {
+                  "items": {
+                    "$ref": "#/components/schemas/Tfl.Api.Presentation.Entities.RoadCorridor"
+                  },
+                  "type": "array"
+                }
+              },
+              "text/json": {
+                "schema": {
+                  "items": {
+                    "$ref": "#/components/schemas/Tfl.Api.Presentation.Entities.RoadCorridor"
+                  },
+                  "type": "array"
+                }
+              },
+              "text/xml": {
+                "schema": {
+                  "items": {
+                    "$ref": "#/components/schemas/Tfl.Api.Presentation.Entities.RoadCorridor"
+                  },
+                  "type": "array"
+                }
+              }
+            },
+            "description": "OK"
+          }
+        },
+        "summary": "Gets all roads managed by TfL",
+        "tags": [
+          "Road"
+        ]
+      }
+    },
+    "/Road/Meta/Categories": {
+      "get": {
+        "deprecated": false,
+        "operationId": "Road_MetaCategories",
+        "responses": {
+          "200": {
+            "content": {
+              "application/json": {
+                "schema": {
+                  "items": {
+                    "type": "string"
+                  },
+                  "type": "array"
+                }
+              },
+              "application/xml": {
+                "schema": {
+                  "items": {
+                    "type": "string"
+                  },
+                  "type": "array"
+                }
+              },
+              "text/json": {
+                "schema": {
+                  "items": {
+                    "type": "string"
+                  },
+                  "type": "array"
+                }
+              },
+              "text/xml": {
+                "schema": {
+                  "items": {
+                    "type": "string"
+                  },
+                  "type": "array"
+                }
+              }
+            },
+            "description": "OK"
+          }
+        },
+        "summary": "Gets a list of valid RoadDisruption categories",
+        "tags": [
+          "Road"
+        ]
+      }
+    },
+    "/Road/Meta/Severities": {
+      "get": {
+        "deprecated": false,
+        "operationId": "Road_MetaSeverities",
+        "responses": {
+          "200": {
+            "content": {
+              "application/json": {
+                "schema": {
+                  "items": {
+                    "$ref": "#/components/schemas/Tfl.Api.Presentation.Entities.StatusSeverity"
+                  },
+                  "type": "array"
+                }
+              },
+              "application/xml": {
+                "schema": {
+                  "items": {
+                    "$ref": "#/components/schemas/Tfl.Api.Presentation.Entities.StatusSeverity"
+                  },
+                  "type": "array"
+                }
+              },
+              "text/json": {
+                "schema": {
+                  "items": {
+                    "$ref": "#/components/schemas/Tfl.Api.Presentation.Entities.StatusSeverity"
+                  },
+                  "type": "array"
+                }
+              },
+              "text/xml": {
+                "schema": {
+                  "items": {
+                    "$ref": "#/components/schemas/Tfl.Api.Presentation.Entities.StatusSeverity"
+                  },
+                  "type": "array"
+                }
+              }
+            },
+            "description": "OK"
+          }
+        },
+        "summary": "Gets a list of valid RoadDisruption severity codes",
+        "tags": [
+          "Road"
+        ]
+      }
+    },
+    "/Road/all/Disruption/{disruptionIds}": {
+      "get": {
+        "deprecated": false,
+        "operationId": "Road_DisruptionById",
+        "parameters": [
+          {
+            "description": "Comma-separated list of disruption identifiers to filter by.",
+            "in": "path",
+            "name": "disruptionIds",
+            "required": true,
+            "schema": {
+              "items": {
+                "type": "string"
+              },
+              "type": "array"
+            },
+            "style": "simple"
+          },
+          {
+            "description": "Optional, defaults to false. When true, removes every property/node except for id, point, severity, severityDescription, startDate, endDate, corridor details, location and comments.",
+            "in": "query",
+            "name": "stripContent",
+            "required": false,
+            "schema": {
+              "type": "boolean"
+            }
+          }
+        ],
+        "responses": {
+          "200": {
+            "content": {
+              "application/geo+json": {
+                "schema": {
+                  "$ref": "#/components/schemas/Tfl.Api.Presentation.Entities.RoadDisruption"
+                }
+              },
+              "application/json": {
+                "schema": {
+                  "$ref": "#/components/schemas/Tfl.Api.Presentation.Entities.RoadDisruption"
+                }
+              },
+              "application/xml": {
+                "schema": {
+                  "$ref": "#/components/schemas/Tfl.Api.Presentation.Entities.RoadDisruption"
+                }
+              },
+              "text/json": {
+                "schema": {
+                  "$ref": "#/components/schemas/Tfl.Api.Presentation.Entities.RoadDisruption"
+                }
+              },
+              "text/xml": {
+                "schema": {
+                  "$ref": "#/components/schemas/Tfl.Api.Presentation.Entities.RoadDisruption"
+                }
+              }
+            },
+            "description": "OK"
+          }
+        },
+        "summary": "Gets a list of active disruptions filtered by disruption Ids.",
+        "tags": [
+          "Road"
+        ]
+      }
+    },
+    "/Road/all/Street/Disruption": {
+      "get": {
+        "deprecated": false,
+        "operationId": "Road_DisruptedStreets",
+        "parameters": [
+          {
+            "description": "Optional, the start time to filter on.",
+            "in": "query",
+            "name": "startDate",
+            "required": true,
+            "schema": {
+              "format": "date-time",
+              "type": "string"
+            }
+          },
+          {
+            "description": "Optional, The end time to filter on.",
+            "in": "query",
+            "name": "endDate",
+            "required": true,
+            "schema": {
+              "format": "date-time",
+              "type": "string"
+            }
+          }
+        ],
+        "responses": {
+          "200": {
+            "content": {
+              "application/json": {
+                "schema": {
+                  "$ref": "#/components/schemas/System.Object"
+                }
+              },
+              "application/xml": {
+                "schema": {
+                  "$ref": "#/components/schemas/System.Object"
+                }
+              },
+              "text/json": {
+                "schema": {
+                  "$ref": "#/components/schemas/System.Object"
+                }
+              },
+              "text/xml": {
+                "schema": {
+                  "$ref": "#/components/schemas/System.Object"
+                }
+              }
+            },
+            "description": "OK"
+          }
+        },
+        "summary": "Gets a list of disrupted streets. If no date filters are provided, current disruptions are returned.",
+        "tags": [
+          "Road"
+        ]
+      }
+    },
+    "/Road/{ids}": {
+      "get": {
+        "deprecated": false,
+        "parameters": [
+          {
+            "description": "Comma-separated list of road identifiers e.g. \"A406, A2\" (a full list of supported road identifiers can be found at the /Road/ endpoint)",
+            "in": "path",
+            "name": "ids",
+            "required": true,
+            "schema": {
+              "items": {
+                "type": "string"
+              },
+              "type": "array"
+            },
+            "style": "simple"
+          }
+        ],
+        "responses": {
+          "200": {
+            "content": {
+              "application/json": {
+                "schema": {
+                  "items": {
+                    "$ref": "#/components/schemas/Tfl.Api.Presentation.Entities.RoadCorridor"
+                  },
+                  "type": "array"
+                }
+              },
+              "application/xml": {
+                "schema": {
+                  "items": {
+                    "$ref": "#/components/schemas/Tfl.Api.Presentation.Entities.RoadCorridor"
+                  },
+                  "type": "array"
+                }
+              },
+              "text/json": {
+                "schema": {
+                  "items": {
+                    "$ref": "#/components/schemas/Tfl.Api.Presentation.Entities.RoadCorridor"
+                  },
+                  "type": "array"
+                }
+              },
+              "text/xml": {
+                "schema": {
+                  "items": {
+                    "$ref": "#/components/schemas/Tfl.Api.Presentation.Entities.RoadCorridor"
+                  },
+                  "type": "array"
+                }
+              }
+            },
+            "description": "OK"
+          }
+        },
+        "summary": "Gets the road with the specified id (e.g. A1)",
+        "tags": [
+          "Road"
+        ]
+      }
+    },
+    "/Road/{ids}/Disruption": {
+      "get": {
+        "deprecated": false,
+        "operationId": "Road_Disruption",
+        "parameters": [
+          {
+            "description": "Comma-separated list of road identifiers e.g. \"A406, A2\" use all for all to ignore id filter (a full list of supported road identifiers can be found at the /Road/ endpoint)",
+            "in": "path",
+            "name": "ids",
+            "required": true,
+            "schema": {
+              "items": {
+                "type": "string"
+              },
+              "type": "array"
+            },
+            "style": "simple"
+          },
+          {
+            "description": "Optional, defaults to false. When true, removes every property/node except for id, point, severity, severityDescription, startDate, endDate, corridor details, location, comments and streets",
+            "in": "query",
+            "name": "stripContent",
+            "required": false,
+            "schema": {
+              "type": "boolean"
+            }
+          },
+          {
+            "description": "an optional list of Severity names to filter on (a valid list of severities can be obtained from the /Road/Meta/severities endpoint)",
+            "explode": true,
+            "in": "query",
+            "name": "severities",
+            "required": false,
+            "schema": {
+              "items": {
+                "type": "string"
+              },
+              "type": "array"
+            }
+          },
+          {
+            "description": "an optional list of category names to filter on (a valid list of categories can be obtained from the /Road/Meta/categories endpoint)",
+            "explode": true,
+            "in": "query",
+            "name": "categories",
+            "required": false,
+            "schema": {
+              "items": {
+                "type": "string"
+              },
+              "type": "array"
+            }
+          },
+          {
+            "description": "Optional, defaults to true. When true, always includes disruptions that have road closures, regardless of the severity filter. When false, the severity filter works as normal.",
+            "in": "query",
+            "name": "closures",
+            "required": false,
+            "schema": {
+              "type": "boolean"
+            }
+          }
+        ],
+        "responses": {
+          "200": {
+            "content": {
+              "application/geo+json": {
+                "schema": {
+                  "items": {
+                    "$ref": "#/components/schemas/Tfl.Api.Presentation.Entities.RoadDisruption"
+                  },
+                  "type": "array"
+                }
+              },
+              "application/json": {
+                "schema": {
+                  "items": {
+                    "$ref": "#/components/schemas/Tfl.Api.Presentation.Entities.RoadDisruption"
+                  },
+                  "type": "array"
+                }
+              },
+              "application/xml": {
+                "schema": {
+                  "items": {
+                    "$ref": "#/components/schemas/Tfl.Api.Presentation.Entities.RoadDisruption"
+                  },
+                  "type": "array"
+                }
+              },
+              "text/json": {
+                "schema": {
+                  "items": {
+                    "$ref": "#/components/schemas/Tfl.Api.Presentation.Entities.RoadDisruption"
+                  },
+                  "type": "array"
+                }
+              },
+              "text/xml": {
+                "schema": {
+                  "items": {
+                    "$ref": "#/components/schemas/Tfl.Api.Presentation.Entities.RoadDisruption"
+                  },
+                  "type": "array"
+                }
+              }
+            },
+            "description": "OK"
+          }
+        },
+        "summary": "Get active disruptions, filtered by road ids",
+        "tags": [
+          "Road"
+        ]
+      }
+    },
+    "/Road/{ids}/Status": {
+      "get": {
+        "deprecated": false,
+        "operationId": "Road_Status",
+        "parameters": [
+          {
+            "description": "Comma-separated list of road identifiers e.g. \"A406, A2\" or use \"all\" to ignore id filter (a full list of supported road identifiers can be found at the /Road/ endpoint)",
+            "in": "path",
+            "name": "ids",
+            "required": true,
+            "schema": {
+              "items": {
+                "type": "string"
+              },
+              "type": "array"
+            },
+            "style": "simple"
+          },
+          {
+            "in": "query",
+            "name": "dateRangeNullable.startDate",
+            "required": false,
+            "schema": {
+              "format": "date-time",
+              "type": "string"
+            }
+          },
+          {
+            "in": "query",
+            "name": "dateRangeNullable.endDate",
+            "required": false,
+            "schema": {
+              "format": "date-time",
+              "type": "string"
+            }
+          }
+        ],
+        "responses": {
+          "200": {
+            "content": {
+              "application/json": {
+                "schema": {
+                  "items": {
+                    "$ref": "#/components/schemas/Tfl.Api.Presentation.Entities.RoadCorridor"
+                  },
+                  "type": "array"
+                }
+              },
+              "application/xml": {
+                "schema": {
+                  "items": {
+                    "$ref": "#/components/schemas/Tfl.Api.Presentation.Entities.RoadCorridor"
+                  },
+                  "type": "array"
+                }
+              },
+              "text/json": {
+                "schema": {
+                  "items": {
+                    "$ref": "#/components/schemas/Tfl.Api.Presentation.Entities.RoadCorridor"
+                  },
+                  "type": "array"
+                }
+              },
+              "text/xml": {
+                "schema": {
+                  "items": {
+                    "$ref": "#/components/schemas/Tfl.Api.Presentation.Entities.RoadCorridor"
+                  },
+                  "type": "array"
+                }
+              }
+            },
+            "description": "OK"
+          }
+        },
+        "summary": "Gets the specified roads with the status aggregated over the date range specified, or now until the end of today if no dates are passed.",
+        "tags": [
+          "Road"
+        ]
+      }
+    },
+    "/Search": {
+      "get": {
+        "deprecated": false,
+        "operationId": "Search_Get",
+        "parameters": [
+          {
+            "description": "The search query",
+            "in": "query",
+            "name": "query",
+            "required": true,
+            "schema": {
+              "type": "string"
+            }
+          }
+        ],
+        "responses": {
+          "200": {
+            "content": {
+              "application/json": {
+                "schema": {
+                  "$ref": "#/components/schemas/Tfl.Api.Presentation.Entities.SearchResponse"
+                }
+              },
+              "application/xml": {
+                "schema": {
+                  "$ref": "#/components/schemas/Tfl.Api.Presentation.Entities.SearchResponse"
+                }
+              },
+              "text/json": {
+                "schema": {
+                  "$ref": "#/components/schemas/Tfl.Api.Presentation.Entities.SearchResponse"
+                }
+              },
+              "text/xml": {
+                "schema": {
+                  "$ref": "#/components/schemas/Tfl.Api.Presentation.Entities.SearchResponse"
+                }
+              }
+            },
+            "description": "OK"
+          }
+        },
+        "summary": "Search the site for occurrences of the query string. The maximum number of results returned is equal to the maximum page size\r\n            of 100. To return subsequent pages, use the paginated overload.",
+        "tags": [
+          "Search"
+        ]
+      }
+    },
+    "/Search/BusSchedules": {
+      "get": {
+        "deprecated": false,
+        "operationId": "Search_BusSchedules",
+        "parameters": [
+          {
+            "description": "The search query",
+            "in": "query",
+            "name": "query",
+            "required": true,
+            "schema": {
+              "type": "string"
+            }
+          }
+        ],
+        "responses": {
+          "200": {
+            "content": {
+              "application/json": {
+                "schema": {
+                  "$ref": "#/components/schemas/Tfl.Api.Presentation.Entities.SearchResponse"
+                }
+              },
+              "application/xml": {
+                "schema": {
+                  "$ref": "#/components/schemas/Tfl.Api.Presentation.Entities.SearchResponse"
+                }
+              },
+              "text/json": {
+                "schema": {
+                  "$ref": "#/components/schemas/Tfl.Api.Presentation.Entities.SearchResponse"
+                }
+              },
+              "text/xml": {
+                "schema": {
+                  "$ref": "#/components/schemas/Tfl.Api.Presentation.Entities.SearchResponse"
+                }
+              }
+            },
+            "description": "OK"
+          }
+        },
+        "summary": "Searches the bus schedules folder on S3 for a given bus number.",
+        "tags": [
+          "Search"
+        ]
+      }
+    },
+    "/Search/Meta/Categories": {
+      "get": {
+        "deprecated": false,
+        "operationId": "Search_MetaCategories",
+        "responses": {
+          "200": {
+            "content": {
+              "application/json": {
+                "schema": {
+                  "items": {
+                    "type": "string"
+                  },
+                  "type": "array"
+                }
+              },
+              "application/xml": {
+                "schema": {
+                  "items": {
+                    "type": "string"
+                  },
+                  "type": "array"
+                }
+              },
+              "text/json": {
+                "schema": {
+                  "items": {
+                    "type": "string"
+                  },
+                  "type": "array"
+                }
+              },
+              "text/xml": {
+                "schema": {
+                  "items": {
+                    "type": "string"
+                  },
+                  "type": "array"
+                }
+              }
+            },
+            "description": "OK"
+          }
+        },
+        "summary": "Gets the available search categories.",
+        "tags": [
+          "Search"
+        ]
+      }
+    },
+    "/Search/Meta/SearchProviders": {
+      "get": {
+        "deprecated": false,
+        "operationId": "Search_MetaSearchProviders",
+        "responses": {
+          "200": {
+            "content": {
+              "application/json": {
+                "schema": {
+                  "items": {
+                    "type": "string"
+                  },
+                  "type": "array"
+                }
+              },
+              "application/xml": {
+                "schema": {
+                  "items": {
+                    "type": "string"
+                  },
+                  "type": "array"
+                }
+              },
+              "text/json": {
+                "schema": {
+                  "items": {
+                    "type": "string"
+                  },
+                  "type": "array"
+                }
+              },
+              "text/xml": {
+                "schema": {
+                  "items": {
+                    "type": "string"
+                  },
+                  "type": "array"
+                }
+              }
+            },
+            "description": "OK"
+          }
+        },
+        "summary": "Gets the available searchProvider names.",
+        "tags": [
+          "Search"
+        ]
+      }
+    },
+    "/Search/Meta/Sorts": {
+      "get": {
+        "deprecated": false,
+        "operationId": "Search_MetaSorts",
+        "responses": {
+          "200": {
+            "content": {
+              "application/json": {
+                "schema": {
+                  "items": {
+                    "type": "string"
+                  },
+                  "type": "array"
+                }
+              },
+              "application/xml": {
+                "schema": {
+                  "items": {
+                    "type": "string"
+                  },
+                  "type": "array"
+                }
+              },
+              "text/json": {
+                "schema": {
+                  "items": {
+                    "type": "string"
+                  },
+                  "type": "array"
+                }
+              },
+              "text/xml": {
+                "schema": {
+                  "items": {
+                    "type": "string"
+                  },
+                  "type": "array"
+                }
+              }
+            },
+            "description": "OK"
+          }
+        },
+        "summary": "Gets the available sorting options.",
+        "tags": [
+          "Search"
+        ]
+      }
+    },
+    "/StopPoint": {
+      "get": {
+        "deprecated": false,
+        "operationId": "StopPoint_GetByGeoPoint",
+        "parameters": [
+          {
+            "description": "a list of stopTypes that should be returned (a list of valid stop types can be obtained from the StopPoint/meta/stoptypes endpoint)",
+            "explode": true,
+            "in": "query",
+            "name": "stopTypes",
+            "required": true,
+            "schema": {
+              "items": {
+                "type": "string"
+              },
+              "type": "array"
+            }
+          },
+          {
+            "description": "the radius of the bounding circle in metres (default : 200)",
+            "in": "query",
+            "name": "radius",
+            "required": false,
+            "schema": {
+              "format": "int32",
+              "type": "integer"
+            }
+          },
+          {
+            "description": "Re-arrange the output into a parent/child hierarchy",
+            "in": "query",
+            "name": "useStopPointHierarchy",
+            "required": false,
+            "schema": {
+              "type": "boolean"
+            }
+          },
+          {
+            "description": "the list of modes to search (comma separated mode names e.g. tube,dlr)",
+            "explode": true,
+            "in": "query",
+            "name": "modes",
+            "required": false,
+            "schema": {
+              "items": {
+                "type": "string"
+              },
+              "type": "array"
+            }
+          },
+          {
+            "description": "an optional list of comma separated property categories to return in the StopPoint's property bag. If null or empty, all categories of property are returned. Pass the keyword \"none\" to return no properties (a valid list of categories can be obtained from the /StopPoint/Meta/categories endpoint)",
+            "explode": true,
+            "in": "query",
+            "name": "categories",
+            "required": false,
+            "schema": {
+              "items": {
+                "type": "string"
+              },
+              "type": "array"
+            }
+          },
+          {
+            "description": "true to return the lines that each stop point serves as a nested resource",
+            "in": "query",
+            "name": "returnLines",
+            "required": false,
+            "schema": {
+              "type": "boolean"
+            }
+          },
+          {
+            "in": "query",
+            "name": "location.lat",
+            "required": true,
+            "schema": {
+              "format": "double",
+              "type": "number"
+            }
+          },
+          {
+            "in": "query",
+            "name": "location.lon",
+            "required": true,
+            "schema": {
+              "format": "double",
+              "type": "number"
+            }
+          }
+        ],
+        "responses": {
+          "200": {
+            "content": {
+              "application/json": {
+                "schema": {
+                  "$ref": "#/components/schemas/Tfl.Api.Presentation.Entities.StopPointsResponse"
+                }
+              },
+              "application/xml": {
+                "schema": {
+                  "$ref": "#/components/schemas/Tfl.Api.Presentation.Entities.StopPointsResponse"
+                }
+              },
+              "text/json": {
+                "schema": {
+                  "$ref": "#/components/schemas/Tfl.Api.Presentation.Entities.StopPointsResponse"
+                }
+              },
+              "text/xml": {
+                "schema": {
+                  "$ref": "#/components/schemas/Tfl.Api.Presentation.Entities.StopPointsResponse"
+                }
+              }
+            },
+            "description": "OK"
+          }
+        },
+        "summary": "Gets a list of StopPoints within {radius} by the specified criteria",
+        "tags": [
+          "StopPoint"
+        ]
+      }
+    },
+    "/StopPoint/Meta/Categories": {
+      "get": {
+        "deprecated": false,
+        "operationId": "StopPoint_MetaCategories",
+        "responses": {
+          "200": {
+            "content": {
+              "application/json": {
+                "schema": {
+                  "items": {
+                    "$ref": "#/components/schemas/Tfl.Api.Presentation.Entities.StopPointCategory"
+                  },
+                  "type": "array"
+                }
+              },
+              "application/xml": {
+                "schema": {
+                  "items": {
+                    "$ref": "#/components/schemas/Tfl.Api.Presentation.Entities.StopPointCategory"
+                  },
+                  "type": "array"
+                }
+              },
+              "text/json": {
+                "schema": {
+                  "items": {
+                    "$ref": "#/components/schemas/Tfl.Api.Presentation.Entities.StopPointCategory"
+                  },
+                  "type": "array"
+                }
+              },
+              "text/xml": {
+                "schema": {
+                  "items": {
+                    "$ref": "#/components/schemas/Tfl.Api.Presentation.Entities.StopPointCategory"
+                  },
+                  "type": "array"
+                }
+              }
+            },
+            "description": "OK"
+          }
+        },
+        "summary": "Gets the list of available StopPoint additional information categories",
+        "tags": [
+          "StopPoint"
+        ]
+      }
+    },
+    "/StopPoint/Meta/Modes": {
+      "get": {
+        "deprecated": false,
+        "operationId": "StopPoint_MetaModes",
+        "responses": {
+          "200": {
+            "content": {
+              "application/json": {
+                "schema": {
+                  "items": {
+                    "$ref": "#/components/schemas/Tfl.Api.Presentation.Entities.Mode"
+                  },
+                  "type": "array"
+                }
+              },
+              "application/xml": {
+                "schema": {
+                  "items": {
+                    "$ref": "#/components/schemas/Tfl.Api.Presentation.Entities.Mode"
+                  },
+                  "type": "array"
+                }
+              },
+              "text/json": {
+                "schema": {
+                  "items": {
+                    "$ref": "#/components/schemas/Tfl.Api.Presentation.Entities.Mode"
+                  },
+                  "type": "array"
+                }
+              },
+              "text/xml": {
+                "schema": {
+                  "items": {
+                    "$ref": "#/components/schemas/Tfl.Api.Presentation.Entities.Mode"
+                  },
+                  "type": "array"
+                }
+              }
+            },
+            "description": "OK"
+          }
+        },
+        "summary": "Gets the list of available StopPoint modes",
+        "tags": [
+          "StopPoint"
+        ]
+      }
+    },
+    "/StopPoint/Meta/StopTypes": {
+      "get": {
+        "deprecated": false,
+        "operationId": "StopPoint_MetaStopTypes",
+        "responses": {
+          "200": {
+            "content": {
+              "application/json": {
+                "schema": {
+                  "items": {
+                    "type": "string"
+                  },
+                  "type": "array"
+                }
+              },
+              "application/xml": {
+                "schema": {
+                  "items": {
+                    "type": "string"
+                  },
+                  "type": "array"
+                }
+              },
+              "text/json": {
+                "schema": {
+                  "items": {
+                    "type": "string"
+                  },
+                  "type": "array"
+                }
+              },
+              "text/xml": {
+                "schema": {
+                  "items": {
+                    "type": "string"
+                  },
+                  "type": "array"
+                }
+              }
+            },
+            "description": "OK"
+          }
+        },
+        "summary": "Gets the list of available StopPoint types",
+        "tags": [
+          "StopPoint"
+        ]
+      }
+    },
+    "/StopPoint/Mode/{modes}": {
+      "get": {
+        "deprecated": false,
+        "operationId": "StopPoint_GetByMode",
+        "parameters": [
+          {
+            "description": "A comma-seperated list of modes e.g. tube,dlr",
+            "in": "path",
+            "name": "modes",
+            "required": true,
+            "schema": {
+              "items": {
+                "type": "string"
+              },
+              "type": "array"
+            },
+            "style": "simple"
+          },
+          {
+            "description": "The data set page to return. Page 1 equates to the first 1000 stop points, page 2 equates to 1001-2000 etc. Must be entered for bus mode as data set is too large.",
+            "in": "query",
+            "name": "page",
+            "required": false,
+            "schema": {
+              "format": "int32",
+              "type": "integer"
+            }
+          }
+        ],
+        "responses": {
+          "200": {
+            "content": {
+              "application/json": {
+                "schema": {
+                  "$ref": "#/components/schemas/Tfl.Api.Presentation.Entities.StopPointsResponse"
+                }
+              },
+              "application/xml": {
+                "schema": {
+                  "$ref": "#/components/schemas/Tfl.Api.Presentation.Entities.StopPointsResponse"
+                }
+              },
+              "text/json": {
+                "schema": {
+                  "$ref": "#/components/schemas/Tfl.Api.Presentation.Entities.StopPointsResponse"
+                }
+              },
+              "text/xml": {
+                "schema": {
+                  "$ref": "#/components/schemas/Tfl.Api.Presentation.Entities.StopPointsResponse"
+                }
+              }
+            },
+            "description": "OK"
+          }
+        },
+        "summary": "Gets a list of StopPoints filtered by the modes available at that StopPoint.",
+        "tags": [
+          "StopPoint"
+        ]
+      }
+    },
+    "/StopPoint/Mode/{modes}/Disruption": {
+      "get": {
+        "deprecated": false,
+        "operationId": "StopPoint_DisruptionByMode",
+        "parameters": [
+          {
+            "description": "A comma-seperated list of modes e.g. tube,dlr",
+            "in": "path",
+            "name": "modes",
+            "required": true,
+            "schema": {
+              "items": {
+                "type": "string"
+              },
+              "type": "array"
+            },
+            "style": "simple"
+          },
+          {
+            "description": "",
+            "in": "query",
+            "name": "includeRouteBlockedStops",
+            "required": false,
+            "schema": {
+              "type": "boolean"
+            }
+          }
+        ],
+        "responses": {
+          "200": {
+            "content": {
+              "application/json": {
+                "schema": {
+                  "items": {
+                    "$ref": "#/components/schemas/Tfl.Api.Presentation.Entities.DisruptedPoint"
+                  },
+                  "type": "array"
+                }
+              },
+              "application/xml": {
+                "schema": {
+                  "items": {
+                    "$ref": "#/components/schemas/Tfl.Api.Presentation.Entities.DisruptedPoint"
+                  },
+                  "type": "array"
+                }
+              },
+              "text/json": {
+                "schema": {
+                  "items": {
+                    "$ref": "#/components/schemas/Tfl.Api.Presentation.Entities.DisruptedPoint"
+                  },
+                  "type": "array"
+                }
+              },
+              "text/xml": {
+                "schema": {
+                  "items": {
+                    "$ref": "#/components/schemas/Tfl.Api.Presentation.Entities.DisruptedPoint"
+                  },
+                  "type": "array"
+                }
+              }
+            },
+            "description": "OK"
+          }
+        },
+        "summary": "Gets a distinct list of disrupted stop points for the given modes",
+        "tags": [
+          "StopPoint"
+        ]
+      }
+    },
+    "/StopPoint/Search": {
+      "get": {
+        "deprecated": false,
+        "parameters": [
+          {
+            "description": "The query string, case-insensitive. Leading and trailing wildcards are applied automatically.",
+            "in": "query",
+            "name": "query",
+            "required": true,
+            "schema": {
+              "type": "string"
+            }
+          },
+          {
+            "description": "An optional, parameter separated list of the modes to filter by",
+            "explode": true,
+            "in": "query",
+            "name": "modes",
+            "required": false,
+            "schema": {
+              "items": {
+                "type": "string"
+              },
+              "type": "array"
+            }
+          },
+          {
+            "description": "True to only return stations in that have Fares data available for single fares to another station.",
+            "in": "query",
+            "name": "faresOnly",
+            "required": false,
+            "schema": {
+              "type": "boolean"
+            }
+          },
+          {
+            "description": "An optional result limit, defaulting to and with a maximum of 50. Since children of the stop point heirarchy are returned for matches,\r\n            it is possible that the flattened result set will contain more than 50 items.",
+            "in": "query",
+            "name": "maxResults",
+            "required": false,
+            "schema": {
+              "format": "int32",
+              "type": "integer"
+            }
+          },
+          {
+            "description": "An optional, parameter separated list of the lines to filter by",
+            "explode": true,
+            "in": "query",
+            "name": "lines",
+            "required": false,
+            "schema": {
+              "items": {
+                "type": "string"
+              },
+              "type": "array"
+            }
+          },
+          {
+            "description": "If true, returns results including HUBs.",
+            "in": "query",
+            "name": "includeHubs",
+            "required": false,
+            "schema": {
+              "type": "boolean"
+            }
+          },
+          {
+            "description": "If the national-rail mode is included, this flag will filter the national rail stations so that only those operated by TfL are returned",
+            "in": "query",
+            "name": "tflOperatedNationalRailStationsOnly",
+            "required": false,
+            "schema": {
+              "type": "boolean"
+            }
+          }
+        ],
+        "responses": {
+          "200": {
+            "content": {
+              "application/json": {
+                "schema": {
+                  "$ref": "#/components/schemas/Tfl.Api.Presentation.Entities.SearchResponse"
+                }
+              },
+              "application/xml": {
+                "schema": {
+                  "$ref": "#/components/schemas/Tfl.Api.Presentation.Entities.SearchResponse"
+                }
+              },
+              "text/json": {
+                "schema": {
+                  "$ref": "#/components/schemas/Tfl.Api.Presentation.Entities.SearchResponse"
+                }
+              },
+              "text/xml": {
+                "schema": {
+                  "$ref": "#/components/schemas/Tfl.Api.Presentation.Entities.SearchResponse"
+                }
+              }
+            },
+            "description": "OK"
+          }
+        },
+        "summary": "Search StopPoints by their common name, or their 5-digit Countdown Bus Stop Code.",
+        "tags": [
+          "StopPoint"
+        ]
+      }
+    },
+    "/StopPoint/Search/{query}": {
+      "get": {
+        "deprecated": false,
+        "operationId": "StopPoint_Search",
+        "parameters": [
+          {
+            "description": "The query string, case-insensitive. Leading and trailing wildcards are applied automatically.",
+            "in": "path",
+            "name": "query",
+            "required": true,
+            "schema": {
+              "type": "string"
+            }
+          },
+          {
+            "description": "An optional, parameter separated list of the modes to filter by",
+            "explode": true,
+            "in": "query",
+            "name": "modes",
+            "required": false,
+            "schema": {
+              "items": {
+                "type": "string"
+              },
+              "type": "array"
+            }
+          },
+          {
+            "description": "True to only return stations in that have Fares data available for single fares to another station.",
+            "in": "query",
+            "name": "faresOnly",
+            "required": false,
+            "schema": {
+              "type": "boolean"
+            }
+          },
+          {
+            "description": "An optional result limit, defaulting to and with a maximum of 50. Since children of the stop point heirarchy are returned for matches,\r\n            it is possible that the flattened result set will contain more than 50 items.",
+            "in": "query",
+            "name": "maxResults",
+            "required": false,
+            "schema": {
+              "format": "int32",
+              "type": "integer"
+            }
+          },
+          {
+            "description": "An optional, parameter separated list of the lines to filter by",
+            "explode": true,
+            "in": "query",
+            "name": "lines",
+            "required": false,
+            "schema": {
+              "items": {
+                "type": "string"
+              },
+              "type": "array"
+            }
+          },
+          {
+            "description": "If true, returns results including HUBs.",
+            "in": "query",
+            "name": "includeHubs",
+            "required": false,
+            "schema": {
+              "type": "boolean"
+            }
+          },
+          {
+            "description": "If the national-rail mode is included, this flag will filter the national rail stations so that only those operated by TfL are returned",
+            "in": "query",
+            "name": "tflOperatedNationalRailStationsOnly",
+            "required": false,
+            "schema": {
+              "type": "boolean"
+            }
+          }
+        ],
+        "responses": {
+          "200": {
+            "content": {
+              "application/json": {
+                "schema": {
+                  "$ref": "#/components/schemas/Tfl.Api.Presentation.Entities.SearchResponse"
+                }
+              },
+              "application/xml": {
+                "schema": {
+                  "$ref": "#/components/schemas/Tfl.Api.Presentation.Entities.SearchResponse"
+                }
+              },
+              "text/json": {
+                "schema": {
+                  "$ref": "#/components/schemas/Tfl.Api.Presentation.Entities.SearchResponse"
+                }
+              },
+              "text/xml": {
+                "schema": {
+                  "$ref": "#/components/schemas/Tfl.Api.Presentation.Entities.SearchResponse"
+                }
+              }
+            },
+            "description": "OK"
+          }
+        },
+        "summary": "Search StopPoints by their common name, or their 5-digit Countdown Bus Stop Code.",
+        "tags": [
+          "StopPoint"
+        ]
+      }
+    },
+    "/StopPoint/ServiceTypes": {
+      "get": {
+        "deprecated": false,
+        "operationId": "StopPoint_GetServiceTypes",
+        "parameters": [
+          {
+            "description": "The Naptan id of the stop",
+            "in": "query",
+            "name": "id",
+            "required": true,
+            "schema": {
+              "type": "string"
+            }
+          },
+          {
+            "description": "The lines which contain the given Naptan id (all lines relevant to the given stoppoint if empty)",
+            "explode": true,
+            "in": "query",
+            "name": "lineIds",
+            "required": false,
+            "schema": {
+              "items": {
+                "type": "string"
+              },
+              "type": "array"
+            }
+          },
+          {
+            "description": "The modes which the lines are relevant to (all if empty)",
+            "explode": true,
+            "in": "query",
+            "name": "modes",
+            "required": false,
+            "schema": {
+              "items": {
+                "type": "string"
+              },
+              "type": "array"
+            }
+          }
+        ],
+        "responses": {
+          "200": {
+            "content": {
+              "application/json": {
+                "schema": {
+                  "items": {
+                    "$ref": "#/components/schemas/Tfl.Api.Presentation.Entities.LineServiceType"
+                  },
+                  "type": "array"
+                }
+              },
+              "application/xml": {
+                "schema": {
+                  "items": {
+                    "$ref": "#/components/schemas/Tfl.Api.Presentation.Entities.LineServiceType"
+                  },
+                  "type": "array"
+                }
+              },
+              "text/json": {
+                "schema": {
+                  "items": {
+                    "$ref": "#/components/schemas/Tfl.Api.Presentation.Entities.LineServiceType"
+                  },
+                  "type": "array"
+                }
+              },
+              "text/xml": {
+                "schema": {
+                  "items": {
+                    "$ref": "#/components/schemas/Tfl.Api.Presentation.Entities.LineServiceType"
+                  },
+                  "type": "array"
+                }
+              }
+            },
+            "description": "OK"
+          }
+        },
+        "summary": "Gets the service types for a given stoppoint",
+        "tags": [
+          "StopPoint"
+        ]
+      }
+    },
+    "/StopPoint/Sms/{id}": {
+      "get": {
+        "deprecated": false,
+        "operationId": "StopPoint_GetBySms",
+        "parameters": [
+          {
+            "description": "A 5-digit Countdown Bus Stop Code e.g. 73241, 50435, 56334.",
+            "in": "path",
+            "name": "id",
+            "required": true,
+            "schema": {
+              "type": "string"
+            }
+          },
+          {
+            "description": "If set to \"web\", a 302 redirect to relevant website bus stop page is returned. Valid values are : web. All other values are ignored.",
+            "in": "query",
+            "name": "output",
+            "required": false,
+            "schema": {
+              "type": "string"
+            }
+          }
+        ],
+        "responses": {
+          "200": {
+            "content": {
+              "application/json": {
+                "schema": {
+                  "$ref": "#/components/schemas/System.Object"
+                }
+              },
+              "application/xml": {
+                "schema": {
+                  "$ref": "#/components/schemas/System.Object"
+                }
+              },
+              "text/json": {
+                "schema": {
+                  "$ref": "#/components/schemas/System.Object"
+                }
+              },
+              "text/xml": {
+                "schema": {
+                  "$ref": "#/components/schemas/System.Object"
+                }
+              }
+            },
+            "description": "OK"
+          }
+        },
+        "summary": "Gets a StopPoint for a given sms code.",
+        "tags": [
+          "StopPoint"
+        ]
+      }
+    },
+    "/StopPoint/Type/{types}": {
+      "get": {
+        "deprecated": false,
+        "operationId": "StopPoint_GetByType",
+        "parameters": [
+          {
+            "description": "A comma-separated list of the types to return. Max. approx. 12 types. \r\n            A list of valid stop types can be obtained from the StopPoint/meta/stoptypes endpoint.",
+            "in": "path",
+            "name": "types",
+            "required": true,
+            "schema": {
+              "items": {
+                "type": "string"
+              },
+              "type": "array"
+            },
+            "style": "simple"
+          }
+        ],
+        "responses": {
+          "200": {
+            "content": {
+              "application/json": {
+                "schema": {
+                  "items": {
+                    "$ref": "#/components/schemas/Tfl.Api.Presentation.Entities.StopPoint"
+                  },
+                  "type": "array"
+                }
+              },
+              "application/xml": {
+                "schema": {
+                  "items": {
+                    "$ref": "#/components/schemas/Tfl.Api.Presentation.Entities.StopPoint"
+                  },
+                  "type": "array"
+                }
+              },
+              "text/json": {
+                "schema": {
+                  "items": {
+                    "$ref": "#/components/schemas/Tfl.Api.Presentation.Entities.StopPoint"
+                  },
+                  "type": "array"
+                }
+              },
+              "text/xml": {
+                "schema": {
+                  "items": {
+                    "$ref": "#/components/schemas/Tfl.Api.Presentation.Entities.StopPoint"
+                  },
+                  "type": "array"
+                }
+              }
+            },
+            "description": "OK"
+          }
+        },
+        "summary": "Gets all stop points of a given type",
+        "tags": [
+          "StopPoint"
+        ]
+      }
+    },
+    "/StopPoint/Type/{types}/page/{page}": {
+      "get": {
+        "deprecated": false,
+        "operationId": "StopPoint_GetByTypeWithPagination",
+        "parameters": [
+          {
+            "description": "",
+            "in": "path",
+            "name": "types",
+            "required": true,
+            "schema": {
+              "items": {
+                "type": "string"
+              },
+              "type": "array"
+            },
+            "style": "simple"
+          },
+          {
+            "description": "",
+            "in": "path",
+            "name": "page",
+            "required": true,
+            "schema": {
+              "format": "int32",
+              "type": "integer"
+            }
+          }
+        ],
+        "responses": {
+          "200": {
+            "content": {
+              "application/json": {
+                "schema": {
+                  "items": {
+                    "$ref": "#/components/schemas/Tfl.Api.Presentation.Entities.StopPoint"
+                  },
+                  "type": "array"
+                }
+              },
+              "application/xml": {
+                "schema": {
+                  "items": {
+                    "$ref": "#/components/schemas/Tfl.Api.Presentation.Entities.StopPoint"
+                  },
+                  "type": "array"
+                }
+              },
+              "text/json": {
+                "schema": {
+                  "items": {
+                    "$ref": "#/components/schemas/Tfl.Api.Presentation.Entities.StopPoint"
+                  },
+                  "type": "array"
+                }
+              },
+              "text/xml": {
+                "schema": {
+                  "items": {
+                    "$ref": "#/components/schemas/Tfl.Api.Presentation.Entities.StopPoint"
+                  },
+                  "type": "array"
+                }
+              }
+            },
+            "description": "OK"
+          }
+        },
+        "summary": "Gets all the stop points of given type(s) with a page number",
+        "tags": [
+          "StopPoint"
+        ]
+      }
+    },
+    "/StopPoint/{ids}": {
+      "get": {
+        "deprecated": false,
+        "operationId": "StopPoint_Get",
+        "parameters": [
+          {
+            "description": "A comma-separated list of stop point ids (station naptan code e.g. 940GZZLUASL). Max. approx. 20 ids.\r\n            You can use /StopPoint/Search/{query} endpoint to find a stop point id from a station name.",
+            "in": "path",
+            "name": "ids",
+            "required": true,
+            "schema": {
+              "items": {
+                "type": "string"
+              },
+              "type": "array"
+            },
+            "style": "simple"
+          },
+          {
+            "description": "Include the crowding data (static). To Filter further use: /StopPoint/{ids}/Crowding/{line}",
+            "in": "query",
+            "name": "includeCrowdingData",
+            "required": false,
+            "schema": {
+              "type": "boolean"
+            }
+          }
+        ],
+        "responses": {
+          "200": {
+            "content": {
+              "application/json": {
+                "schema": {
+                  "items": {
+                    "$ref": "#/components/schemas/Tfl.Api.Presentation.Entities.StopPoint"
+                  },
+                  "type": "array"
+                }
+              },
+              "application/xml": {
+                "schema": {
+                  "items": {
+                    "$ref": "#/components/schemas/Tfl.Api.Presentation.Entities.StopPoint"
+                  },
+                  "type": "array"
+                }
+              },
+              "text/json": {
+                "schema": {
+                  "items": {
+                    "$ref": "#/components/schemas/Tfl.Api.Presentation.Entities.StopPoint"
+                  },
+                  "type": "array"
+                }
+              },
+              "text/xml": {
+                "schema": {
+                  "items": {
+                    "$ref": "#/components/schemas/Tfl.Api.Presentation.Entities.StopPoint"
+                  },
+                  "type": "array"
+                }
+              }
+            },
+            "description": "OK"
+          }
+        },
+        "summary": "Gets a list of StopPoints corresponding to the given list of stop ids.",
+        "tags": [
+          "StopPoint"
+        ]
+      }
+    },
+    "/StopPoint/{ids}/Disruption": {
+      "get": {
+        "deprecated": false,
+        "operationId": "StopPoint_Disruption",
+        "parameters": [
+          {
+            "description": "A comma-seperated list of stop point ids. Max. approx. 20 ids.\r\n            You can use /StopPoint/Search/{query} endpoint to find a stop point id from a station name.",
+            "in": "path",
+            "name": "ids",
+            "required": true,
+            "schema": {
+              "items": {
+                "type": "string"
+              },
+              "type": "array"
+            },
+            "style": "simple"
+          },
+          {
+            "description": "Specify true to return disruptions for entire family, or false to return disruptions for just this stop point. Defaults to false.",
+            "in": "query",
+            "name": "getFamily",
+            "required": false,
+            "schema": {
+              "type": "boolean"
+            }
+          },
+          {
+            "description": "",
+            "in": "query",
+            "name": "includeRouteBlockedStops",
+            "required": false,
+            "schema": {
+              "type": "boolean"
+            }
+          },
+          {
+            "description": "Specify true to associate all disruptions with parent stop point. (Only applicable when getFamily is true).",
+            "in": "query",
+            "name": "flattenResponse",
+            "required": false,
+            "schema": {
+              "type": "boolean"
+            }
+          }
+        ],
+        "responses": {
+          "200": {
+            "content": {
+              "application/json": {
+                "schema": {
+                  "items": {
+                    "$ref": "#/components/schemas/Tfl.Api.Presentation.Entities.DisruptedPoint"
+                  },
+                  "type": "array"
+                }
+              },
+              "application/xml": {
+                "schema": {
+                  "items": {
+                    "$ref": "#/components/schemas/Tfl.Api.Presentation.Entities.DisruptedPoint"
+                  },
+                  "type": "array"
+                }
+              },
+              "text/json": {
+                "schema": {
+                  "items": {
+                    "$ref": "#/components/schemas/Tfl.Api.Presentation.Entities.DisruptedPoint"
+                  },
+                  "type": "array"
+                }
+              },
+              "text/xml": {
+                "schema": {
+                  "items": {
+                    "$ref": "#/components/schemas/Tfl.Api.Presentation.Entities.DisruptedPoint"
+                  },
+                  "type": "array"
+                }
+              }
+            },
+            "description": "OK"
+          }
+        },
+        "summary": "Gets all disruptions for the specified StopPointId, plus disruptions for any child Naptan records it may have.",
+        "tags": [
+          "StopPoint"
+        ]
+      }
+    },
+    "/StopPoint/{id}/ArrivalDepartures": {
+      "get": {
+        "deprecated": false,
+        "operationId": "StopPoint_ArrivalDepartures",
+        "parameters": [
+          {
+            "description": "A StopPoint id (station naptan code e.g. 940GZZLUASL, you can use /StopPoint/Search/{query} endpoint to find a stop point id from a station name)",
+            "in": "path",
+            "name": "id",
+            "required": true,
+            "schema": {
+              "type": "string"
+            }
+          },
+          {
+            "description": "A comma-separated list of line ids e.g. tfl-rail, london-overground, thameslink",
+            "explode": true,
+            "in": "query",
+            "name": "lineIds",
+            "required": true,
+            "schema": {
+              "items": {
+                "type": "string"
+              },
+              "type": "array"
+            }
+          }
+        ],
+        "responses": {
+          "200": {
+            "content": {
+              "application/json": {
+                "schema": {
+                  "items": {
+                    "$ref": "#/components/schemas/Tfl.Api.Presentation.Entities.ArrivalDeparture"
+                  },
+                  "type": "array"
+                }
+              },
+              "application/xml": {
+                "schema": {
+                  "items": {
+                    "$ref": "#/components/schemas/Tfl.Api.Presentation.Entities.ArrivalDeparture"
+                  },
+                  "type": "array"
+                }
+              },
+              "text/json": {
+                "schema": {
+                  "items": {
+                    "$ref": "#/components/schemas/Tfl.Api.Presentation.Entities.ArrivalDeparture"
+                  },
+                  "type": "array"
+                }
+              },
+              "text/xml": {
+                "schema": {
+                  "items": {
+                    "$ref": "#/components/schemas/Tfl.Api.Presentation.Entities.ArrivalDeparture"
+                  },
+                  "type": "array"
+                }
+              }
+            },
+            "description": "OK"
+          }
+        },
+        "summary": "Gets the list of arrival and departure predictions for the given stop point id (overground, tfl rail and thameslink only)",
+        "tags": [
+          "StopPoint"
+        ]
+      }
+    },
+    "/StopPoint/{id}/Arrivals": {
+      "get": {
+        "deprecated": false,
+        "operationId": "StopPoint_Arrivals",
+        "parameters": [
+          {
+            "description": "A StopPoint id (station naptan code e.g. 940GZZLUASL, you can use /StopPoint/Search/{query} endpoint to find a stop point id from a station name)",
+            "in": "path",
+            "name": "id",
+            "required": true,
+            "schema": {
+              "type": "string"
+            }
+          }
+        ],
+        "responses": {
+          "200": {
+            "content": {
+              "application/json": {
+                "schema": {
+                  "items": {
+                    "$ref": "#/components/schemas/Tfl.Api.Presentation.Entities.Prediction"
+                  },
+                  "type": "array"
+                }
+              },
+              "application/xml": {
+                "schema": {
+                  "items": {
+                    "$ref": "#/components/schemas/Tfl.Api.Presentation.Entities.Prediction"
+                  },
+                  "type": "array"
+                }
+              },
+              "text/json": {
+                "schema": {
+                  "items": {
+                    "$ref": "#/components/schemas/Tfl.Api.Presentation.Entities.Prediction"
+                  },
+                  "type": "array"
+                }
+              },
+              "text/xml": {
+                "schema": {
+                  "items": {
+                    "$ref": "#/components/schemas/Tfl.Api.Presentation.Entities.Prediction"
+                  },
+                  "type": "array"
+                }
+              }
+            },
+            "description": "OK"
+          }
+        },
+        "summary": "Gets the list of arrival predictions for the given stop point id",
+        "tags": [
+          "StopPoint"
+        ]
+      }
+    },
+    "/StopPoint/{id}/CanReachOnLine/{lineId}": {
+      "get": {
+        "deprecated": false,
+        "operationId": "StopPoint_ReachableFrom",
+        "parameters": [
+          {
+            "description": "The id (station naptan code e.g. 940GZZLUASL, you can use /StopPoint/Search/{query} endpoint to find a stop point id from a station name) of the stop point to filter by",
+            "in": "path",
+            "name": "id",
+            "required": true,
+            "schema": {
+              "type": "string"
+            }
+          },
+          {
+            "description": "Line id of the line to filter by (e.g. victoria)",
+            "in": "path",
+            "name": "lineId",
+            "required": true,
+            "schema": {
+              "type": "string"
+            }
+          },
+          {
+            "description": "A comma-separated list of service types to filter on. If not specified. Supported values: Regular, Night. Defaulted to 'Regular' if not specified",
+            "explode": true,
+            "in": "query",
+            "name": "serviceTypes",
+            "required": false,
+            "schema": {
+              "items": {
+                "enum": [
+                  "Regular",
+                  "Night"
+                ],
+                "type": "string"
+              },
+              "type": "array"
+            }
+          }
+        ],
+        "responses": {
+          "200": {
+            "content": {
+              "application/json": {
+                "schema": {
+                  "items": {
+                    "$ref": "#/components/schemas/Tfl.Api.Presentation.Entities.StopPoint"
+                  },
+                  "type": "array"
+                }
+              },
+              "application/xml": {
+                "schema": {
+                  "items": {
+                    "$ref": "#/components/schemas/Tfl.Api.Presentation.Entities.StopPoint"
+                  },
+                  "type": "array"
+                }
+              },
+              "text/json": {
+                "schema": {
+                  "items": {
+                    "$ref": "#/components/schemas/Tfl.Api.Presentation.Entities.StopPoint"
+                  },
+                  "type": "array"
+                }
+              },
+              "text/xml": {
+                "schema": {
+                  "items": {
+                    "$ref": "#/components/schemas/Tfl.Api.Presentation.Entities.StopPoint"
+                  },
+                  "type": "array"
+                }
+              }
+            },
+            "description": "OK"
+          }
+        },
+        "summary": "Gets Stopoints that are reachable from a station/line combination.",
+        "tags": [
+          "StopPoint"
+        ]
+      }
+    },
+    "/StopPoint/{id}/Crowding/{line}": {
+      "get": {
+        "deprecated": false,
+        "operationId": "StopPoint_Crowding",
+        "parameters": [
+          {
+            "description": "The Naptan id of the stop",
+            "in": "path",
+            "name": "id",
+            "required": true,
+            "schema": {
+              "type": "string"
+            }
+          },
+          {
+            "description": "A particular line e.g. victoria, circle, northern etc.",
+            "in": "path",
+            "name": "line",
+            "required": true,
+            "schema": {
+              "type": "string"
+            }
+          },
+          {
+            "description": "The direction of travel. Can be inbound or outbound.",
+            "in": "query",
+            "name": "direction",
+            "required": true,
+            "schema": {
+              "enum": [
+                "inbound",
+                "outbound",
+                "all"
+              ],
+              "type": "string"
+            }
+          }
+        ],
+        "responses": {
+          "200": {
+            "content": {
+              "application/json": {
+                "schema": {
+                  "items": {
+                    "$ref": "#/components/schemas/Tfl.Api.Presentation.Entities.StopPoint"
+                  },
+                  "type": "array"
+                }
+              },
+              "application/xml": {
+                "schema": {
+                  "items": {
+                    "$ref": "#/components/schemas/Tfl.Api.Presentation.Entities.StopPoint"
+                  },
+                  "type": "array"
+                }
+              },
+              "text/json": {
+                "schema": {
+                  "items": {
+                    "$ref": "#/components/schemas/Tfl.Api.Presentation.Entities.StopPoint"
+                  },
+                  "type": "array"
+                }
+              },
+              "text/xml": {
+                "schema": {
+                  "items": {
+                    "$ref": "#/components/schemas/Tfl.Api.Presentation.Entities.StopPoint"
+                  },
+                  "type": "array"
+                }
+              }
+            },
+            "description": "OK"
+          }
+        },
+        "summary": "Gets all the Crowding data (static) for the StopPointId, plus crowding data for a given line and optionally a particular direction.",
+        "tags": [
+          "StopPoint"
+        ]
+      }
+    },
+    "/StopPoint/{id}/DirectionTo/{toStopPointId}": {
+      "get": {
+        "deprecated": false,
+        "operationId": "StopPoint_Direction",
+        "parameters": [
+          {
+            "description": "Originating stop id (station naptan code e.g. 940GZZLUASL, you can use /StopPoint/Search/{query} endpoint to find a stop point id from a station name)",
+            "in": "path",
+            "name": "id",
+            "required": true,
+            "schema": {
+              "type": "string"
+            }
+          },
+          {
+            "description": "Destination stop id (station naptan code e.g. 940GZZLUASL, you can use /StopPoint/Search/{query} endpoint to find a stop point id from a station name)",
+            "in": "path",
+            "name": "toStopPointId",
+            "required": true,
+            "schema": {
+              "type": "string"
+            }
+          },
+          {
+            "description": "Optional line id filter e.g. victoria",
+            "in": "query",
+            "name": "lineId",
+            "required": false,
+            "schema": {
+              "type": "string"
+            }
+          }
+        ],
+        "responses": {
+          "200": {
+            "content": {
+              "application/json": {
+                "schema": {
+                  "type": "string"
+                }
+              },
+              "application/xml": {
+                "schema": {
+                  "type": "string"
+                }
+              },
+              "text/json": {
+                "schema": {
+                  "type": "string"
+                }
+              },
+              "text/xml": {
+                "schema": {
+                  "type": "string"
+                }
+              }
+            },
+            "description": "OK"
+          }
+        },
+        "summary": "Returns the canonical direction, \"inbound\" or \"outbound\", for a given pair of stop point Ids in the direction from -&gt; to.",
+        "tags": [
+          "StopPoint"
+        ]
+      }
+    },
+    "/StopPoint/{id}/Route": {
+      "get": {
+        "deprecated": false,
+        "operationId": "StopPoint_Route",
+        "parameters": [
+          {
+            "description": "A stop point id (station naptan codes e.g. 940GZZLUASL, you can use /StopPoint/Search/{query} endpoint to find a stop point id from a station name)",
+            "in": "path",
+            "name": "id",
+            "required": true,
+            "schema": {
+              "type": "string"
+            }
+          },
+          {
+            "description": "A comma-separated list of service types to filter on. If not specified. Supported values: Regular, Night. Defaulted to 'Regular' if not specified",
+            "explode": true,
+            "in": "query",
+            "name": "serviceTypes",
+            "required": false,
+            "schema": {
+              "items": {
+                "enum": [
+                  "Regular",
+                  "Night"
+                ],
+                "type": "string"
+              },
+              "type": "array"
+            }
+          }
+        ],
+        "responses": {
+          "200": {
+            "content": {
+              "application/json": {
+                "schema": {
+                  "items": {
+                    "$ref": "#/components/schemas/Tfl.Api.Presentation.Entities.StopPointRouteSection"
+                  },
+                  "type": "array"
+                }
+              },
+              "application/xml": {
+                "schema": {
+                  "items": {
+                    "$ref": "#/components/schemas/Tfl.Api.Presentation.Entities.StopPointRouteSection"
+                  },
+                  "type": "array"
+                }
+              },
+              "text/json": {
+                "schema": {
+                  "items": {
+                    "$ref": "#/components/schemas/Tfl.Api.Presentation.Entities.StopPointRouteSection"
+                  },
+                  "type": "array"
+                }
+              },
+              "text/xml": {
+                "schema": {
+                  "items": {
+                    "$ref": "#/components/schemas/Tfl.Api.Presentation.Entities.StopPointRouteSection"
+                  },
+                  "type": "array"
+                }
+              }
+            },
+            "description": "OK"
+          }
+        },
+        "summary": "Returns the route sections for all the lines that service the given stop point ids",
+        "tags": [
+          "StopPoint"
+        ]
+      }
+    },
+    "/StopPoint/{id}/placeTypes": {
+      "get": {
+        "deprecated": false,
+        "parameters": [
+          {
+            "description": "A naptan id for a stop point (station naptan code e.g. 940GZZLUASL).",
+            "in": "path",
+            "name": "id",
+            "required": true,
+            "schema": {
+              "type": "string"
+            }
+          },
+          {
+            "description": "A comcomma-separated value representing the place types.",
+            "explode": true,
+            "in": "query",
+            "name": "placeTypes",
+            "required": true,
+            "schema": {
+              "items": {
+                "type": "string"
+              },
+              "type": "array"
+            }
+          }
+        ],
+        "responses": {
+          "200": {
+            "content": {
+              "application/json": {
+                "schema": {
+                  "items": {
+                    "$ref": "#/components/schemas/Tfl.Api.Presentation.Entities.Place"
+                  },
+                  "type": "array"
+                }
+              },
+              "application/xml": {
+                "schema": {
+                  "items": {
+                    "$ref": "#/components/schemas/Tfl.Api.Presentation.Entities.Place"
+                  },
+                  "type": "array"
+                }
+              },
+              "text/json": {
+                "schema": {
+                  "items": {
+                    "$ref": "#/components/schemas/Tfl.Api.Presentation.Entities.Place"
+                  },
+                  "type": "array"
+                }
+              },
+              "text/xml": {
+                "schema": {
+                  "items": {
+                    "$ref": "#/components/schemas/Tfl.Api.Presentation.Entities.Place"
+                  },
+                  "type": "array"
+                }
+              }
+            },
+            "description": "OK"
+          }
+        },
+        "summary": "Get a list of places corresponding to a given id and place types.",
+        "tags": [
+          "StopPoint"
+        ]
+      }
+    },
+    "/StopPoint/{stopPointId}/CarParks": {
+      "get": {
+        "deprecated": false,
+        "operationId": "StopPoint_GetCarParksById",
+        "parameters": [
+          {
+            "description": "stopPointId is required to get the car parks.",
+            "in": "path",
+            "name": "stopPointId",
+            "required": true,
+            "schema": {
+              "type": "string"
+            }
+          }
+        ],
+        "responses": {
+          "200": {
+            "content": {
+              "application/json": {
+                "schema": {
+                  "items": {
+                    "$ref": "#/components/schemas/Tfl.Api.Presentation.Entities.Place"
+                  },
+                  "type": "array"
+                }
+              },
+              "application/xml": {
+                "schema": {
+                  "items": {
+                    "$ref": "#/components/schemas/Tfl.Api.Presentation.Entities.Place"
+                  },
+                  "type": "array"
+                }
+              },
+              "text/json": {
+                "schema": {
+                  "items": {
+                    "$ref": "#/components/schemas/Tfl.Api.Presentation.Entities.Place"
+                  },
+                  "type": "array"
+                }
+              },
+              "text/xml": {
+                "schema": {
+                  "items": {
+                    "$ref": "#/components/schemas/Tfl.Api.Presentation.Entities.Place"
+                  },
+                  "type": "array"
+                }
+              }
+            },
+            "description": "OK"
+          }
+        },
+        "summary": "Get car parks corresponding to the given stop point id.",
+        "tags": [
+          "StopPoint"
+        ]
+      }
+    },
+    "/StopPoint/{stopPointId}/TaxiRanks": {
+      "get": {
+        "deprecated": false,
+        "operationId": "StopPoint_GetTaxiRanksByIds",
+        "parameters": [
+          {
+            "description": "stopPointId is required to get the taxi ranks.",
+            "in": "path",
+            "name": "stopPointId",
+            "required": true,
+            "schema": {
+              "type": "string"
+            }
+          }
+        ],
+        "responses": {
+          "200": {
+            "content": {
+              "application/json": {
+                "schema": {
+                  "items": {
+                    "$ref": "#/components/schemas/Tfl.Api.Presentation.Entities.Place"
+                  },
+                  "type": "array"
+                }
+              },
+              "application/xml": {
+                "schema": {
+                  "items": {
+                    "$ref": "#/components/schemas/Tfl.Api.Presentation.Entities.Place"
+                  },
+                  "type": "array"
+                }
+              },
+              "text/json": {
+                "schema": {
+                  "items": {
+                    "$ref": "#/components/schemas/Tfl.Api.Presentation.Entities.Place"
+                  },
+                  "type": "array"
+                }
+              },
+              "text/xml": {
+                "schema": {
+                  "items": {
+                    "$ref": "#/components/schemas/Tfl.Api.Presentation.Entities.Place"
+                  },
+                  "type": "array"
+                }
+              }
+            },
+            "description": "OK"
+          }
+        },
+        "summary": "Gets a list of taxi ranks corresponding to the given stop point id.",
+        "tags": [
+          "StopPoint"
+        ]
+      }
+    },
+    "/TravelTimes/compareOverlay/{z}/mapcenter/{mapCenterLat}/{mapCenterLon}/pinlocation/{pinLat}/{pinLon}/dimensions/{width}/{height}": {
+      "get": {
+        "deprecated": false,
+        "operationId": "TravelTime_GetCompareOverlay",
+        "parameters": [
+          {
+            "description": "The zoom level.",
+            "in": "path",
+            "name": "z",
+            "required": true,
+            "schema": {
+              "format": "int32",
+              "type": "integer"
+            }
+          },
+          {
+            "description": "The latitude of the pin.",
+            "in": "path",
+            "name": "pinLat",
+            "required": true,
+            "schema": {
+              "format": "double",
+              "type": "number"
+            }
+          },
+          {
+            "description": "The longitude of the pin.",
+            "in": "path",
+            "name": "pinLon",
+            "required": true,
+            "schema": {
+              "format": "double",
+              "type": "number"
+            }
+          },
+          {
+            "description": "The map center latitude.",
+            "in": "path",
+            "name": "mapCenterLat",
+            "required": true,
+            "schema": {
+              "format": "double",
+              "type": "number"
+            }
+          },
+          {
+            "description": "The map center longitude.",
+            "in": "path",
+            "name": "mapCenterLon",
+            "required": true,
+            "schema": {
+              "format": "double",
+              "type": "number"
+            }
+          },
+          {
+            "description": "The title of the scenario.",
+            "in": "query",
+            "name": "scenarioTitle",
+            "required": true,
+            "schema": {
+              "type": "string"
+            }
+          },
+          {
+            "description": "The id for the time of day (AM/INTER/PM)",
+            "in": "query",
+            "name": "timeOfDayId",
+            "required": true,
+            "schema": {
+              "type": "string"
+            }
+          },
+          {
+            "description": "The id of the mode.",
+            "in": "query",
+            "name": "modeId",
+            "required": true,
+            "schema": {
+              "type": "string"
+            }
+          },
+          {
+            "description": "The width of the requested overlay.",
+            "in": "path",
+            "name": "width",
+            "required": true,
+            "schema": {
+              "format": "int32",
+              "type": "integer"
+            }
+          },
+          {
+            "description": "The height of the requested overlay.",
+            "in": "path",
+            "name": "height",
+            "required": true,
+            "schema": {
+              "format": "int32",
+              "type": "integer"
+            }
+          },
+          {
+            "description": "The direction of travel.",
+            "in": "query",
+            "name": "direction",
+            "required": true,
+            "schema": {
+              "enum": [
+                "Average",
+                "From",
+                "To"
+              ],
+              "type": "string"
+            }
+          },
+          {
+            "description": "The total minutes between the travel time bands",
+            "in": "query",
+            "name": "travelTimeInterval",
+            "required": true,
+            "schema": {
+              "format": "int32",
+              "type": "integer"
+            }
+          },
+          {
+            "in": "query",
+            "name": "compareType",
+            "required": true,
+            "schema": {
+              "type": "string"
+            }
+          },
+          {
+            "in": "query",
+            "name": "compareValue",
+            "required": true,
+            "schema": {
+              "type": "string"
+            }
+          }
+        ],
+        "responses": {
+          "200": {
+            "content": {
+              "application/json": {
+                "schema": {
+                  "$ref": "#/components/schemas/System.Object"
+                }
+              },
+              "application/xml": {
+                "schema": {
+                  "$ref": "#/components/schemas/System.Object"
+                }
+              },
+              "text/json": {
+                "schema": {
+                  "$ref": "#/components/schemas/System.Object"
+                }
+              },
+              "text/xml": {
+                "schema": {
+                  "$ref": "#/components/schemas/System.Object"
+                }
+              }
+            },
+            "description": "OK"
+          }
+        },
+        "summary": "Gets the TravelTime overlay.",
+        "tags": [
+          "TravelTime"
+        ]
+      }
+    },
+    "/TravelTimes/overlay/{z}/mapcenter/{mapCenterLat}/{mapCenterLon}/pinlocation/{pinLat}/{pinLon}/dimensions/{width}/{height}": {
+      "get": {
+        "deprecated": false,
+        "operationId": "TravelTime_GetOverlay",
+        "parameters": [
+          {
+            "description": "The zoom level.",
+            "in": "path",
+            "name": "z",
+            "required": true,
+            "schema": {
+              "format": "int32",
+              "type": "integer"
+            }
+          },
+          {
+            "description": "The latitude of the pin.",
+            "in": "path",
+            "name": "pinLat",
+            "required": true,
+            "schema": {
+              "format": "double",
+              "type": "number"
+            }
+          },
+          {
+            "description": "The longitude of the pin.",
+            "in": "path",
+            "name": "pinLon",
+            "required": true,
+            "schema": {
+              "format": "double",
+              "type": "number"
+            }
+          },
+          {
+            "description": "The map center latitude.",
+            "in": "path",
+            "name": "mapCenterLat",
+            "required": true,
+            "schema": {
+              "format": "double",
+              "type": "number"
+            }
+          },
+          {
+            "description": "The map center longitude.",
+            "in": "path",
+            "name": "mapCenterLon",
+            "required": true,
+            "schema": {
+              "format": "double",
+              "type": "number"
+            }
+          },
+          {
+            "description": "The title of the scenario.",
+            "in": "query",
+            "name": "scenarioTitle",
+            "required": true,
+            "schema": {
+              "type": "string"
+            }
+          },
+          {
+            "description": "The id for the time of day (AM/INTER/PM)",
+            "in": "query",
+            "name": "timeOfDayId",
+            "required": true,
+            "schema": {
+              "type": "string"
+            }
+          },
+          {
+            "description": "The id of the mode.",
+            "in": "query",
+            "name": "modeId",
+            "required": true,
+            "schema": {
+              "type": "string"
+            }
+          },
+          {
+            "description": "The width of the requested overlay.",
+            "in": "path",
+            "name": "width",
+            "required": true,
+            "schema": {
+              "format": "int32",
+              "type": "integer"
+            }
+          },
+          {
+            "description": "The height of the requested overlay.",
+            "in": "path",
+            "name": "height",
+            "required": true,
+            "schema": {
+              "format": "int32",
+              "type": "integer"
+            }
+          },
+          {
+            "description": "The direction of travel.",
+            "in": "query",
+            "name": "direction",
+            "required": true,
+            "schema": {
+              "enum": [
+                "Average",
+                "From",
+                "To"
+              ],
+              "type": "string"
+            }
+          },
+          {
+            "description": "The total minutes between the travel time bands",
+            "in": "query",
+            "name": "travelTimeInterval",
+            "required": true,
+            "schema": {
+              "format": "int32",
+              "type": "integer"
+            }
+          }
+        ],
+        "responses": {
+          "200": {
+            "content": {
+              "application/json": {
+                "schema": {
+                  "$ref": "#/components/schemas/System.Object"
+                }
+              },
+              "application/xml": {
+                "schema": {
+                  "$ref": "#/components/schemas/System.Object"
+                }
+              },
+              "text/json": {
+                "schema": {
+                  "$ref": "#/components/schemas/System.Object"
+                }
+              },
+              "text/xml": {
+                "schema": {
+                  "$ref": "#/components/schemas/System.Object"
+                }
+              }
+            },
+            "description": "OK"
+          }
+        },
+        "summary": "Gets the TravelTime overlay.",
+        "tags": [
+          "TravelTime"
+        ]
+      }
+    },
+    "/Vehicle/EmissionSurcharge": {
+      "get": {
+        "deprecated": false,
+        "operationId": "Vehicle_GetEmissionsSurchargeCompliance",
+        "parameters": [
+          {
+            "description": "The Vehicle Registration Mark",
+            "in": "query",
+            "name": "vrm",
+            "required": true,
+            "schema": {
+              "type": "string"
+            }
+          }
+        ],
+        "responses": {
+          "200": {
+            "content": {
+              "application/json": {
+                "schema": {
+                  "$ref": "#/components/schemas/Tfl.Api.Presentation.Entities.VehicleMatch"
+                }
+              },
+              "application/xml": {
+                "schema": {
+                  "$ref": "#/components/schemas/Tfl.Api.Presentation.Entities.VehicleMatch"
+                }
+              },
+              "text/json": {
+                "schema": {
+                  "$ref": "#/components/schemas/Tfl.Api.Presentation.Entities.VehicleMatch"
+                }
+              },
+              "text/xml": {
+                "schema": {
+                  "$ref": "#/components/schemas/Tfl.Api.Presentation.Entities.VehicleMatch"
+                }
+              }
+            },
+            "description": "OK"
+          }
+        },
+        "summary": "Gets the Emissions Surcharge compliance for the Vehicle",
+        "tags": [
+          "Vehicle"
+        ]
+      }
+    },
+    "/Vehicle/UlezCompliance": {
+      "get": {
+        "deprecated": false,
+        "operationId": "Vehicle_GetUlezCompliance",
+        "parameters": [
+          {
+            "description": "The Vehicle Registration Mark",
+            "in": "query",
+            "name": "vrm",
+            "required": true,
+            "schema": {
+              "type": "string"
+            }
+          }
+        ],
+        "responses": {
+          "200": {
+            "content": {
+              "application/json": {
+                "schema": {
+                  "$ref": "#/components/schemas/Tfl.Api.Presentation.Entities.VehicleMatch"
+                }
+              },
+              "application/xml": {
+                "schema": {
+                  "$ref": "#/components/schemas/Tfl.Api.Presentation.Entities.VehicleMatch"
+                }
+              },
+              "text/json": {
+                "schema": {
+                  "$ref": "#/components/schemas/Tfl.Api.Presentation.Entities.VehicleMatch"
+                }
+              },
+              "text/xml": {
+                "schema": {
+                  "$ref": "#/components/schemas/Tfl.Api.Presentation.Entities.VehicleMatch"
+                }
+              }
+            },
+            "description": "OK"
+          }
+        },
+        "summary": "Gets the Ulez Surcharge compliance for the Vehicle",
+        "tags": [
+          "Vehicle"
+        ]
+      }
+    },
+    "/Vehicle/{ids}/Arrivals": {
+      "get": {
+        "deprecated": false,
+        "operationId": "Vehicle_Get",
+        "parameters": [
+          {
+            "description": "A comma-separated list of vehicle ids e.g. LX58CFV,LX11AZB,LX58CFE. Max approx. 25 ids.",
+            "in": "path",
+            "name": "ids",
+            "required": true,
+            "schema": {
+              "items": {
+                "type": "string"
+              },
+              "type": "array"
+            },
+            "style": "simple"
+          }
+        ],
+        "responses": {
+          "200": {
+            "content": {
+              "application/json": {
+                "schema": {
+                  "items": {
+                    "$ref": "#/components/schemas/Tfl.Api.Presentation.Entities.Prediction"
+                  },
+                  "type": "array"
+                }
+              },
+              "application/xml": {
+                "schema": {
+                  "items": {
+                    "$ref": "#/components/schemas/Tfl.Api.Presentation.Entities.Prediction"
+                  },
+                  "type": "array"
+                }
+              },
+              "text/json": {
+                "schema": {
+                  "items": {
+                    "$ref": "#/components/schemas/Tfl.Api.Presentation.Entities.Prediction"
+                  },
+                  "type": "array"
+                }
+              },
+              "text/xml": {
+                "schema": {
+                  "items": {
+                    "$ref": "#/components/schemas/Tfl.Api.Presentation.Entities.Prediction"
+                  },
+                  "type": "array"
+                }
+              }
+            },
+            "description": "OK"
+          }
+        },
+        "summary": "Gets the predictions for a given list of vehicle Id's.",
+        "tags": [
+          "Vehicle"
+        ]
+      }
+    }
+  },
+  "components": {
+    "schemas": {
+      "System.Data.Spatial.DbGeography": {
+        "properties": {
+          "geography": {
+            "$ref": "#/components/schemas/System.Data.Spatial.DbGeographyWellKnownValue"
+          }
+        },
+        "type": "object"
+      },
+      "System.Data.Spatial.DbGeographyWellKnownValue": {
+        "properties": {
+          "coordinateSystemId": {
+            "format": "int32",
+            "type": "integer"
+          },
+          "wellKnownBinary": {
+            "format": "byte",
+            "type": "string"
+          },
+          "wellKnownText": {
+            "type": "string"
+          }
+        },
+        "type": "object"
+      },
+      "System.Object": {
+        "properties": {},
+        "type": "object"
+      },
+      "Tfl.Api.Common.ApiVersionInfo": {
+        "properties": {
+          "assemblies": {
+            "items": {
+              "type": "string"
+            },
+            "type": "array"
+          },
+          "label": {
+            "type": "string"
+          },
+          "timestamp": {
+            "format": "date-time",
+            "type": "string"
+          },
+          "version": {
+            "type": "string"
+          }
+        },
+        "type": "object"
+      },
+      "Tfl.Api.Common.DateRange": {
+        "properties": {
+          "endDate": {
+            "format": "date-time",
+            "type": "string"
+          },
+          "startDate": {
+            "format": "date-time",
+            "type": "string"
+          }
+        },
+        "type": "object"
+      },
+      "Tfl.Api.Common.DateRangeNullable": {
+        "properties": {
+          "endDate": {
+            "format": "date-time",
+            "type": "string"
+          },
+          "startDate": {
+            "format": "date-time",
+            "type": "string"
+          }
+        },
+        "type": "object"
+      },
+      "Tfl.Api.Common.GeoPoint": {
+        "properties": {
+          "lat": {
+            "format": "double",
+            "type": "number"
+          },
+          "lon": {
+            "format": "double",
+            "type": "number"
+          }
+        },
+        "required": [
+          "lat",
+          "lon"
+        ],
+        "type": "object"
+      },
+      "Tfl.Api.Common.JourneyPlanner.JpElevation": {
+        "properties": {
+          "distance": {
+            "format": "int32",
+            "type": "integer"
+          },
+          "endLat": {
+            "format": "double",
+            "type": "number"
+          },
+          "endLon": {
+            "format": "double",
+            "type": "number"
+          },
+          "gradient": {
+            "format": "double",
+            "type": "number"
+          },
+          "heightFromPreviousPoint": {
+            "format": "int32",
+            "type": "integer"
+          },
+          "startLat": {
+            "format": "double",
+            "type": "number"
+          },
+          "startLon": {
+            "format": "double",
+            "type": "number"
+          }
+        },
+        "type": "object"
+      },
+      "Tfl.Api.Common.PlaceGeo": {
+        "properties": {
+          "lat": {
+            "format": "double",
+            "type": "number"
+          },
+          "lon": {
+            "format": "double",
+            "type": "number"
+          },
+          "neLat": {
+            "format": "double",
+            "type": "number"
+          },
+          "neLon": {
+            "format": "double",
+            "type": "number"
+          },
+          "swLat": {
+            "format": "double",
+            "type": "number"
+          },
+          "swLon": {
+            "format": "double",
+            "type": "number"
+          }
+        },
+        "type": "object"
+      },
+      "Tfl.Api.Common.PostcodeInput": {
+        "properties": {
+          "postcode": {
+            "pattern": "^(([gG][iI][rR] {0,}0[aA]{2})|((([a-pr-uwyzA-PR-UWYZ][a-hk-yA-HK-Y]?[0-9][0-9]?)|(([a-pr-uwyzA-PR-UWYZ][0-9][a-hjkstuwA-HJKSTUW])|([a-pr-uwyzA-PR-UWYZ][a-hk-yA-HK-Y][0-9][abehmnprv-yABEHMNPRV-Y]))) {0,}[0-9][abd-hjlnp-uw-zABD-HJLNP-UW-Z]{2}))$",
+            "type": "string"
+          }
+        },
+        "type": "object"
+      },
+      "Tfl.Api.Presentation.Entities.AccidentStats.AccidentDetail": {
+        "properties": {
+          "borough": {
+            "type": "string"
+          },
+          "casualties": {
+            "items": {
+              "$ref": "#/components/schemas/Tfl.Api.Presentation.Entities.AccidentStats.Casualty"
+            },
+            "type": "array"
+          },
+          "date": {
+            "format": "date-time",
+            "type": "string"
+          },
+          "id": {
+            "format": "int32",
+            "type": "integer"
+          },
+          "lat": {
+            "format": "double",
+            "type": "number"
+          },
+          "location": {
+            "type": "string"
+          },
+          "lon": {
+            "format": "double",
+            "type": "number"
+          },
+          "severity": {
+            "type": "string"
+          },
+          "vehicles": {
+            "items": {
+              "$ref": "#/components/schemas/Tfl.Api.Presentation.Entities.AccidentStats.Vehicle"
+            },
+            "type": "array"
+          }
+        },
+        "type": "object"
+      },
+      "Tfl.Api.Presentation.Entities.AccidentStats.AccidentStatsOrderedSummary": {
+        "properties": {
+          "accidents": {
+            "format": "int32",
+            "type": "integer"
+          },
+          "borough": {
+            "type": "string"
+          },
+          "year": {
+            "format": "int32",
+            "type": "integer"
+          }
+        },
+        "type": "object"
+      },
+      "Tfl.Api.Presentation.Entities.AccidentStats.Casualty": {
+        "properties": {
+          "age": {
+            "format": "int32",
+            "type": "integer"
+          },
+          "ageBand": {
+            "type": "string"
+          },
+          "class": {
+            "type": "string"
+          },
+          "mode": {
+            "type": "string"
+          },
+          "severity": {
+            "type": "string"
+          }
+        },
+        "type": "object"
+      },
+      "Tfl.Api.Presentation.Entities.AccidentStats.Vehicle": {
+        "properties": {
+          "type": {
+            "type": "string"
+          }
+        },
+        "type": "object"
+      },
+      "Tfl.Api.Presentation.Entities.ActiveServiceType": {
+        "properties": {
+          "mode": {
+            "type": "string"
+          },
+          "serviceType": {
+            "type": "string"
+          }
+        },
+        "type": "object"
+      },
+      "Tfl.Api.Presentation.Entities.AdditionalProperties": {
+        "properties": {
+          "category": {
+            "type": "string"
+          },
+          "key": {
+            "type": "string"
+          },
+          "modified": {
+            "format": "date-time",
+            "type": "string"
+          },
+          "sourceSystemKey": {
+            "type": "string"
+          },
+          "value": {
+            "type": "string"
+          }
+        },
+        "type": "object"
+      },
+      "Tfl.Api.Presentation.Entities.ArrivalDeparture": {
+        "description": "DTO to capture the prediction details",
+        "properties": {
+          "cause": {
+            "description": "Reason for cancellation or delay",
+            "type": "string"
+          },
+          "departureStatus": {
+            "description": "Status of departure",
+            "enum": [
+              "OnTime",
+              "Delayed",
+              "Cancelled",
+              "NotStoppingAtStation"
+            ],
+            "type": "string"
+          },
+          "destinationName": {
+            "description": "Name of the destination",
+            "type": "string"
+          },
+          "destinationNaptanId": {
+            "description": "Naptan Identifier for the prediction's destination",
+            "type": "string"
+          },
+          "estimatedTimeOfArrival": {
+            "description": "Estimated time of arrival",
+            "format": "date-time",
+            "type": "string"
+          },
+          "estimatedTimeOfDeparture": {
+            "description": "Estimated time of arrival",
+            "format": "date-time",
+            "type": "string"
+          },
+          "minutesAndSecondsToArrival": {
+            "description": "Estimated time of arrival",
+            "type": "string"
+          },
+          "minutesAndSecondsToDeparture": {
+            "description": "Estimated time of arrival",
+            "type": "string"
+          },
+          "naptanId": {
+            "description": "Identifier for the prediction",
+            "type": "string"
+          },
+          "platformName": {
+            "description": "Platform name (for bus, this is the stop letter)",
+            "type": "string"
+          },
+          "scheduledTimeOfArrival": {
+            "description": "Estimated time of arrival",
+            "format": "date-time",
+            "type": "string"
+          },
+          "scheduledTimeOfDeparture": {
+            "description": "Estimated time of arrival",
+            "format": "date-time",
+            "type": "string"
+          },
+          "stationName": {
+            "description": "Station name",
+            "type": "string"
+          },
+          "timing": {
+            "$ref": "#/components/schemas/Tfl.Api.Presentation.Entities.PredictionTiming"
+          }
+        },
+        "type": "object"
+      },
+      "Tfl.Api.Presentation.Entities.Bay": {
+        "properties": {
+          "bayCount": {
+            "format": "int32",
+            "type": "integer"
+          },
+          "bayType": {
+            "type": "string"
+          },
+          "free": {
+            "format": "int32",
+            "type": "integer"
+          },
+          "occupied": {
+            "format": "int32",
+            "type": "integer"
+          }
+        },
+        "type": "object"
+      },
+      "Tfl.Api.Presentation.Entities.BikePointOccupancy": {
+        "description": "Bike point occupancy",
+        "properties": {
+          "bikesCount": {
+            "description": "Total bike counts",
+            "format": "int32",
+            "type": "integer"
+          },
+          "emptyDocks": {
+            "description": "Empty docks",
+            "format": "int32",
+            "type": "integer"
+          },
+          "id": {
+            "description": "Id of the bike point such as BikePoints_1",
+            "type": "string"
+          },
+          "name": {
+            "description": "Name / Common name of the bike point",
+            "type": "string"
+          },
+          "totalDocks": {
+            "description": "Total docks available",
+            "format": "int32",
+            "type": "integer"
+          }
+        },
+        "type": "object"
+      },
+      "Tfl.Api.Presentation.Entities.CarParkOccupancy": {
+        "properties": {
+          "bays": {
+            "items": {
+              "$ref": "#/components/schemas/Tfl.Api.Presentation.Entities.Bay"
+            },
+            "type": "array"
+          },
+          "carParkDetailsUrl": {
+            "type": "string"
+          },
+          "id": {
+            "type": "string"
+          },
+          "name": {
+            "type": "string"
+          }
+        },
+        "type": "object"
+      },
+      "Tfl.Api.Presentation.Entities.ChargeConnectorOccupancy": {
+        "description": "",
+        "properties": {
+          "id": {
+            "description": "",
+            "format": "int32",
+            "type": "integer"
+          },
+          "sourceSystemPlaceId": {
+            "description": "",
+            "type": "string"
+          },
+          "status": {
+            "description": "",
+            "type": "string"
+          }
+        },
+        "type": "object"
+      },
+      "Tfl.Api.Presentation.Entities.Coordinate": {
+        "properties": {
+          "easting": {
+            "format": "double",
+            "type": "number"
+          },
+          "latitude": {
+            "format": "double",
+            "type": "number"
+          },
+          "longitude": {
+            "format": "double",
+            "type": "number"
+          },
+          "northing": {
+            "format": "double",
+            "type": "number"
+          },
+          "xCoord": {
+            "format": "int32",
+            "type": "integer"
+          },
+          "yCoord": {
+            "format": "int32",
+            "type": "integer"
+          }
+        },
+        "type": "object"
+      },
+      "Tfl.Api.Presentation.Entities.Crowding": {
+        "properties": {
+          "passengerFlows": {
+            "description": "Busiest times at a station (static information)",
+            "items": {
+              "$ref": "#/components/schemas/Tfl.Api.Presentation.Entities.PassengerFlow"
+            },
+            "type": "array"
+          },
+          "trainLoadings": {
+            "description": "Train Loading on a scale 1-6, 1 being \"Very quiet\" and 6 being \"Exceptionally busy\" (static information)",
+            "items": {
+              "$ref": "#/components/schemas/Tfl.Api.Presentation.Entities.TrainLoading"
+            },
+            "type": "array"
+          }
+        },
+        "type": "object"
+      },
+      "Tfl.Api.Presentation.Entities.CycleSuperhighway": {
+        "description": "",
+        "properties": {
+          "geography": {
+            "$ref": "#/components/schemas/System.Data.Spatial.DbGeography"
+          },
+          "id": {
+            "description": "The Id",
+            "type": "string"
+          },
+          "label": {
+            "description": "The long label to show on maps when zoomed in",
+            "type": "string"
+          },
+          "labelShort": {
+            "description": "The short label to show on maps",
+            "type": "string"
+          },
+          "modified": {
+            "description": "When the data was last updated",
+            "format": "date-time",
+            "type": "string"
+          },
+          "routeType": {
+            "description": "Type of cycle route e.g CycleSuperhighways, Quietways, MiniHollands etc",
+            "enum": [
+              "Unknown",
+              "All",
+              "Cycle Superhighways",
+              "Quietways",
+              "Cycleways",
+              "Mini-Hollands",
+              "Central London Grid",
+              "Streetspace Route"
+            ],
+            "type": "string"
+          },
+          "segmented": {
+            "description": "True if the route is split into segments",
+            "type": "boolean"
+          },
+          "status": {
+            "description": "Cycle route status i.e Proposed, Existing etc",
+            "enum": [
+              "Unknown",
+              "All",
+              "Open",
+              "In Progress",
+              "Planned",
+              "Planned - Subject to feasibility and consultation.",
+              "Not Open"
+            ],
+            "type": "string"
+          }
+        },
+        "type": "object"
+      },
+      "Tfl.Api.Presentation.Entities.DisruptedPoint": {
+        "properties": {
+          "additionalInformation": {
+            "type": "string"
+          },
+          "appearance": {
+            "type": "string"
+          },
+          "atcoCode": {
+            "type": "string"
+          },
+          "commonName": {
+            "type": "string"
+          },
+          "description": {
+            "type": "string"
+          },
+          "fromDate": {
+            "format": "date-time",
+            "type": "string"
+          },
+          "mode": {
+            "type": "string"
+          },
+          "stationAtcoCode": {
+            "type": "string"
+          },
+          "toDate": {
+            "format": "date-time",
+            "type": "string"
+          },
+          "type": {
+            "type": "string"
+          }
+        },
+        "type": "object"
+      },
+      "Tfl.Api.Presentation.Entities.Disruption": {
+        "description": "Represents a disruption to a route within the transport network.",
+        "properties": {
+          "additionalInfo": {
+            "description": "Gets or sets the additionaInfo of this disruption.",
+            "type": "string"
+          },
+          "affectedRoutes": {
+            "description": "Gets or sets the routes affected by this disruption",
+            "items": {
+              "$ref": "#/components/schemas/Tfl.Api.Presentation.Entities.RouteSection"
+            },
+            "type": "array"
+          },
+          "affectedStops": {
+            "description": "Gets or sets the stops affected by this disruption",
+            "items": {
+              "$ref": "#/components/schemas/Tfl.Api.Presentation.Entities.StopPoint"
+            },
+            "type": "array"
+          },
+          "category": {
+            "description": "Gets or sets the category of this dispruption.",
+            "enum": [
+              "Undefined",
+              "RealTime",
+              "PlannedWork",
+              "Information",
+              "Event",
+              "Crowding",
+              "StatusAlert"
+            ],
+            "type": "string"
+          },
+          "categoryDescription": {
+            "description": "Gets or sets the description of the category.",
+            "type": "string"
+          },
+          "closureText": {
+            "description": "Text describing the closure type",
+            "type": "string"
+          },
+          "created": {
+            "description": "Gets or sets the date/time when this disruption was created.",
+            "format": "date-time",
+            "type": "string"
+          },
+          "description": {
+            "description": "Gets or sets the description of this disruption.",
+            "type": "string"
+          },
+          "lastUpdate": {
+            "description": "Gets or sets the date/time when this disruption was last updated.",
+            "format": "date-time",
+            "type": "string"
+          },
+          "summary": {
+            "description": "Gets or sets the summary of this disruption.",
+            "type": "string"
+          },
+          "type": {
+            "description": "Gets or sets the disruption type of this dispruption.",
+            "type": "string"
+          }
+        },
+        "type": "object"
+      },
+      "Tfl.Api.Presentation.Entities.Fares.Fare": {
+        "properties": {
+          "cap": {
+            "format": "double",
+            "type": "number"
+          },
+          "cost": {
+            "type": "string"
+          },
+          "description": {
+            "type": "string"
+          },
+          "id": {
+            "format": "int32",
+            "type": "integer"
+          },
+          "mode": {
+            "type": "string"
+          },
+          "passengerType": {
+            "type": "string"
+          },
+          "ticketTime": {
+            "type": "string"
+          },
+          "ticketType": {
+            "type": "string"
+          },
+          "validFrom": {
+            "format": "date-time",
+            "type": "string"
+          },
+          "validUntil": {
+            "format": "date-time",
+            "type": "string"
+          },
+          "zone": {
+            "type": "string"
+          }
+        },
+        "type": "object"
+      },
+      "Tfl.Api.Presentation.Entities.Fares.FareBounds": {
+        "properties": {
+          "description": {
+            "type": "string"
+          },
+          "displayName": {
+            "type": "string"
+          },
+          "displayOrder": {
+            "format": "int32",
+            "type": "integer"
+          },
+          "from": {
+            "type": "string"
+          },
+          "id": {
+            "format": "int32",
+            "type": "integer"
+          },
+          "isPopularFare": {
+            "type": "boolean"
+          },
+          "isPopularTravelCard": {
+            "type": "boolean"
+          },
+          "isTour": {
+            "type": "boolean"
+          },
+          "messages": {
+            "items": {
+              "$ref": "#/components/schemas/Tfl.Api.Presentation.Entities.Message"
+            },
+            "type": "array"
+          },
+          "operator": {
+            "type": "string"
+          },
+          "routeCode": {
+            "type": "string"
+          },
+          "to": {
+            "type": "string"
+          },
+          "via": {
+            "type": "string"
+          }
+        },
+        "type": "object"
+      },
+      "Tfl.Api.Presentation.Entities.Fares.FareDetails": {
+        "properties": {
+          "boundsId": {
+            "format": "int32",
+            "type": "integer"
+          },
+          "contactlessPAYGOnlyFare": {
+            "type": "boolean"
+          },
+          "displayName": {
+            "type": "string"
+          },
+          "displayOrder": {
+            "format": "int32",
+            "type": "integer"
+          },
+          "endDate": {
+            "format": "date-time",
+            "type": "string"
+          },
+          "from": {
+            "type": "string"
+          },
+          "fromStation": {
+            "type": "string"
+          },
+          "isTour": {
+            "type": "boolean"
+          },
+          "messages": {
+            "items": {
+              "$ref": "#/components/schemas/Tfl.Api.Presentation.Entities.Message"
+            },
+            "type": "array"
+          },
+          "mode": {
+            "type": "string"
+          },
+          "operator": {
+            "type": "string"
+          },
+          "passengerType": {
+            "type": "string"
+          },
+          "routeCode": {
+            "type": "string"
+          },
+          "routeDescription": {
+            "type": "string"
+          },
+          "specialFare": {
+            "type": "boolean"
+          },
+          "startDate": {
+            "format": "date-time",
+            "type": "string"
+          },
+          "throughFare": {
+            "type": "boolean"
+          },
+          "ticketsAvailable": {
+            "items": {
+              "$ref": "#/components/schemas/Tfl.Api.Presentation.Entities.Fares.Ticket"
+            },
+            "type": "array"
+          },
+          "to": {
+            "type": "string"
+          },
+          "toStation": {
+            "type": "string"
+          },
+          "validatorInformation": {
+            "type": "string"
+          },
+          "via": {
+            "type": "string"
+          }
+        },
+        "type": "object"
+      },
+      "Tfl.Api.Presentation.Entities.Fares.FareStation": {
+        "properties": {
+          "atcoCode": {
+            "type": "string"
+          },
+          "commonName": {
+            "type": "string"
+          },
+          "fareCategory": {
+            "enum": [
+              "Cash",
+              "Oyster",
+              "Contactless",
+              "ContactlessOnly",
+              "All"
+            ],
+            "type": "string"
+          }
+        },
+        "type": "object"
+      },
+      "Tfl.Api.Presentation.Entities.Fares.FaresMode": {
+        "properties": {
+          "description": {
+            "type": "string"
+          },
+          "id": {
+            "format": "int32",
+            "type": "integer"
+          },
+          "name": {
+            "type": "string"
+          }
+        },
+        "type": "object"
+      },
+      "Tfl.Api.Presentation.Entities.Fares.FaresPeriod": {
+        "properties": {
+          "endDate": {
+            "format": "date-time",
+            "type": "string"
+          },
+          "id": {
+            "format": "int32",
+            "type": "integer"
+          },
+          "isFuture": {
+            "type": "boolean"
+          },
+          "startDate": {
+            "format": "date-time",
+            "type": "string"
+          },
+          "viewableDate": {
+            "format": "date-time",
+            "type": "string"
+          }
+        },
+        "type": "object"
+      },
+      "Tfl.Api.Presentation.Entities.Fares.FaresSection": {
+        "properties": {
+          "header": {
+            "type": "string"
+          },
+          "index": {
+            "format": "int32",
+            "type": "integer"
+          },
+          "journey": {
+            "$ref": "#/components/schemas/Tfl.Api.Presentation.Entities.Fares.Journey"
+          },
+          "messages": {
+            "items": {
+              "$ref": "#/components/schemas/Tfl.Api.Presentation.Entities.Message"
+            },
+            "type": "array"
+          },
+          "rows": {
+            "items": {
+              "$ref": "#/components/schemas/Tfl.Api.Presentation.Entities.Fares.FareDetails"
+            },
+            "type": "array"
+          }
+        },
+        "type": "object"
+      },
+      "Tfl.Api.Presentation.Entities.Fares.Journey": {
+        "properties": {
+          "fromStation": {
+            "$ref": "#/components/schemas/Tfl.Api.Presentation.Entities.Fares.FareStation"
+          },
+          "toStation": {
+            "$ref": "#/components/schemas/Tfl.Api.Presentation.Entities.Fares.FareStation"
+          }
+        },
+        "type": "object"
+      },
+      "Tfl.Api.Presentation.Entities.Fares.PassengerType": {
+        "properties": {
+          "description": {
+            "type": "string"
+          },
+          "displayName": {
+            "type": "string"
+          },
+          "displayOrder": {
+            "format": "int32",
+            "type": "integer"
+          },
+          "type": {
+            "type": "string"
+          }
+        },
+        "type": "object"
+      },
+      "Tfl.Api.Presentation.Entities.Fares.Recommendation": {
+        "properties": {
+          "cost": {
+            "type": "string"
+          },
+          "discountCard": {
+            "type": "string"
+          },
+          "fareType": {
+            "type": "string"
+          },
+          "gettingYourTicket": {
+            "items": {
+              "$ref": "#/components/schemas/Tfl.Api.Presentation.Entities.Message"
+            },
+            "type": "array"
+          },
+          "id": {
+            "format": "int32",
+            "type": "integer"
+          },
+          "keyFeatures": {
+            "items": {
+              "$ref": "#/components/schemas/Tfl.Api.Presentation.Entities.Message"
+            },
+            "type": "array"
+          },
+          "notes": {
+            "items": {
+              "$ref": "#/components/schemas/Tfl.Api.Presentation.Entities.Message"
+            },
+            "type": "array"
+          },
+          "priceComparison": {
+            "type": "string"
+          },
+          "priceDescription": {
+            "type": "string"
+          },
+          "product": {
+            "type": "string"
+          },
+          "productType": {
+            "type": "string"
+          },
+          "rank": {
+            "format": "int32",
+            "type": "integer"
+          },
+          "recommendedTopUp": {
+            "type": "string"
+          },
+          "rule": {
+            "format": "int32",
+            "type": "integer"
+          },
+          "singleFare": {
+            "format": "double",
+            "type": "number"
+          },
+          "ticketTime": {
+            "type": "string"
+          },
+          "ticketType": {
+            "type": "string"
+          },
+          "zones": {
+            "type": "string"
+          }
+        },
+        "type": "object"
+      },
+      "Tfl.Api.Presentation.Entities.Fares.RecommendationResponse": {
+        "properties": {
+          "recommendations": {
+            "items": {
+              "$ref": "#/components/schemas/Tfl.Api.Presentation.Entities.Fares.Recommendation"
+            },
+            "type": "array"
+          }
+        },
+        "type": "object"
+      },
+      "Tfl.Api.Presentation.Entities.Fares.Ticket": {
+        "properties": {
+          "cost": {
+            "type": "string"
+          },
+          "description": {
+            "type": "string"
+          },
+          "displayOrder": {
+            "format": "int32",
+            "type": "integer"
+          },
+          "messages": {
+            "items": {
+              "$ref": "#/components/schemas/Tfl.Api.Presentation.Entities.Message"
+            },
+            "type": "array"
+          },
+          "mode": {
+            "type": "string"
+          },
+          "passengerType": {
+            "type": "string"
+          },
+          "ticketTime": {
+            "$ref": "#/components/schemas/Tfl.Api.Presentation.Entities.Fares.TicketTime"
+          },
+          "ticketType": {
+            "$ref": "#/components/schemas/Tfl.Api.Presentation.Entities.Fares.TicketType"
+          }
+        },
+        "type": "object"
+      },
+      "Tfl.Api.Presentation.Entities.Fares.TicketTime": {
+        "properties": {
+          "description": {
+            "type": "string"
+          },
+          "type": {
+            "type": "string"
+          }
+        },
+        "type": "object"
+      },
+      "Tfl.Api.Presentation.Entities.Fares.TicketType": {
+        "properties": {
+          "description": {
+            "type": "string"
+          },
+          "type": {
+            "type": "string"
+          }
+        },
+        "type": "object"
+      },
+      "Tfl.Api.Presentation.Entities.GeoCodeSearchMatch": {
+        "properties": {
+          "address": {
+            "description": "A string describing the formatted address of the place. Adds additional context to the place's Name.",
+            "type": "string"
+          },
+          "id": {
+            "type": "string"
+          },
+          "lat": {
+            "format": "double",
+            "type": "number"
+          },
+          "lon": {
+            "format": "double",
+            "type": "number"
+          },
+          "name": {
+            "type": "string"
+          },
+          "types": {
+            "description": "The type of the place e.g. \"street_address\"",
+            "items": {
+              "type": "string"
+            },
+            "type": "array"
+          },
+          "url": {
+            "type": "string"
+          }
+        },
+        "type": "object"
+      },
+      "Tfl.Api.Presentation.Entities.Identifier": {
+        "properties": {
+          "crowding": {
+            "$ref": "#/components/schemas/Tfl.Api.Presentation.Entities.Crowding"
+          },
+          "fullName": {
+            "type": "string"
+          },
+          "id": {
+            "type": "string"
+          },
+          "name": {
+            "type": "string"
+          },
+          "routeType": {
+            "enum": [
+              "Unknown",
+              "All",
+              "Cycle Superhighways",
+              "Quietways",
+              "Cycleways",
+              "Mini-Hollands",
+              "Central London Grid",
+              "Streetspace Route"
+            ],
+            "type": "string"
+          },
+          "status": {
+            "enum": [
+              "Unknown",
+              "All",
+              "Open",
+              "In Progress",
+              "Planned",
+              "Planned - Subject to feasibility and consultation.",
+              "Not Open"
+            ],
+            "type": "string"
+          },
+          "type": {
+            "type": "string"
+          },
+          "uri": {
+            "type": "string"
+          }
+        },
+        "type": "object"
+      },
+      "Tfl.Api.Presentation.Entities.Instruction": {
+        "properties": {
+          "detailed": {
+            "type": "string"
+          },
+          "steps": {
+            "items": {
+              "$ref": "#/components/schemas/Tfl.Api.Presentation.Entities.InstructionStep"
+            },
+            "type": "array"
+          },
+          "summary": {
+            "type": "string"
+          }
+        },
+        "type": "object"
+      },
+      "Tfl.Api.Presentation.Entities.InstructionStep": {
+        "properties": {
+          "cumulativeDistance": {
+            "format": "int32",
+            "type": "integer"
+          },
+          "cumulativeTravelTime": {
+            "format": "int32",
+            "type": "integer"
+          },
+          "description": {
+            "type": "string"
+          },
+          "descriptionHeading": {
+            "type": "string"
+          },
+          "distance": {
+            "format": "int32",
+            "type": "integer"
+          },
+          "latitude": {
+            "format": "double",
+            "type": "number"
+          },
+          "longitude": {
+            "format": "double",
+            "type": "number"
+          },
+          "pathAttribute": {
+            "$ref": "#/components/schemas/Tfl.Api.Presentation.Entities.PathAttribute"
+          },
+          "skyDirection": {
+            "format": "int32",
+            "type": "integer"
+          },
+          "skyDirectionDescription": {
+            "enum": [
+              "North",
+              "NorthEast",
+              "East",
+              "SouthEast",
+              "South",
+              "SouthWest",
+              "West",
+              "NorthWest"
+            ],
+            "type": "string"
+          },
+          "streetName": {
+            "type": "string"
+          },
+          "trackType": {
+            "enum": [
+              "CycleSuperHighway",
+              "CanalTowpath",
+              "QuietRoad",
+              "ProvisionForCyclists",
+              "BusyRoads",
+              "None",
+              "PushBike",
+              "Quietway"
+            ],
+            "type": "string"
+          },
+          "turnDirection": {
+            "type": "string"
+          }
+        },
+        "type": "object"
+      },
+      "Tfl.Api.Presentation.Entities.Interval": {
+        "properties": {
+          "stopId": {
+            "type": "string"
+          },
+          "timeToArrival": {
+            "format": "double",
+            "type": "number"
+          }
+        },
+        "type": "object"
+      },
+      "Tfl.Api.Presentation.Entities.JourneyPlanner.Fare": {
+        "properties": {
+          "chargeLevel": {
+            "type": "string"
+          },
+          "chargeProfileName": {
+            "type": "string"
+          },
+          "cost": {
+            "format": "int32",
+            "type": "integer"
+          },
+          "highZone": {
+            "format": "int32",
+            "type": "integer"
+          },
+          "isHopperFare": {
+            "type": "boolean"
+          },
+          "lowZone": {
+            "format": "int32",
+            "type": "integer"
+          },
+          "offPeak": {
+            "format": "int32",
+            "type": "integer"
+          },
+          "peak": {
+            "format": "int32",
+            "type": "integer"
+          },
+          "taps": {
+            "items": {
+              "$ref": "#/components/schemas/Tfl.Api.Presentation.Entities.JourneyPlanner.FareTap"
+            },
+            "type": "array"
+          }
+        },
+        "type": "object"
+      },
+      "Tfl.Api.Presentation.Entities.JourneyPlanner.FareCaveat": {
+        "properties": {
+          "text": {
+            "type": "string"
+          },
+          "type": {
+            "type": "string"
+          }
+        },
+        "type": "object"
+      },
+      "Tfl.Api.Presentation.Entities.JourneyPlanner.FareTap": {
+        "properties": {
+          "atcoCode": {
+            "type": "string"
+          },
+          "tapDetails": {
+            "$ref": "#/components/schemas/Tfl.Api.Presentation.Entities.JourneyPlanner.FareTapDetails"
+          }
+        },
+        "type": "object"
+      },
+      "Tfl.Api.Presentation.Entities.JourneyPlanner.FareTapDetails": {
+        "properties": {
+          "busRouteId": {
+            "type": "string"
+          },
+          "hostDeviceType": {
+            "type": "string"
+          },
+          "modeType": {
+            "type": "string"
+          },
+          "nationalLocationCode": {
+            "format": "int32",
+            "type": "integer"
+          },
+          "tapTimestamp": {
+            "format": "date-time",
+            "type": "string"
+          },
+          "validationType": {
+            "type": "string"
+          }
+        },
+        "type": "object"
+      },
+      "Tfl.Api.Presentation.Entities.JourneyPlanner.ItineraryResult": {
+        "description": "A DTO representing a list of possible journeys.",
+        "properties": {
+          "cycleHireDockingStationData": {
+            "$ref": "#/components/schemas/Tfl.Api.Presentation.Entities.JourneyPlanner.JourneyPlannerCycleHireDockingStationData"
+          },
+          "journeyVector": {
+            "$ref": "#/components/schemas/Tfl.Api.Presentation.Entities.JourneyPlanner.JourneyVector"
+          },
+          "journeys": {
+            "items": {
+              "$ref": "#/components/schemas/Tfl.Api.Presentation.Entities.JourneyPlanner.Journey"
+            },
+            "type": "array"
+          },
+          "lines": {
+            "items": {
+              "$ref": "#/components/schemas/Tfl.Api.Presentation.Entities.Line"
+            },
+            "type": "array"
+          },
+          "recommendedMaxAgeMinutes": {
+            "format": "int32",
+            "type": "integer"
+          },
+          "searchCriteria": {
+            "$ref": "#/components/schemas/Tfl.Api.Presentation.Entities.JourneyPlanner.SearchCriteria"
+          },
+          "stopMessages": {
+            "items": {
+              "type": "string"
+            },
+            "type": "array"
+          }
+        },
+        "type": "object"
+      },
+      "Tfl.Api.Presentation.Entities.JourneyPlanner.Journey": {
+        "description": "Object that represents an end to end journey (see schematic).",
+        "properties": {
+          "arrivalDateTime": {
+            "format": "date-time",
+            "type": "string"
+          },
+          "duration": {
+            "format": "int32",
+            "type": "integer"
+          },
+          "fare": {
+            "$ref": "#/components/schemas/Tfl.Api.Presentation.Entities.JourneyPlanner.JourneyFare"
+          },
+          "legs": {
+            "items": {
+              "$ref": "#/components/schemas/Tfl.Api.Presentation.Entities.JourneyPlanner.Leg"
+            },
+            "type": "array"
+          },
+          "startDateTime": {
+            "format": "date-time",
+            "type": "string"
+          }
+        },
+        "type": "object"
+      },
+      "Tfl.Api.Presentation.Entities.JourneyPlanner.JourneyFare": {
+        "properties": {
+          "caveats": {
+            "items": {
+              "$ref": "#/components/schemas/Tfl.Api.Presentation.Entities.JourneyPlanner.FareCaveat"
+            },
+            "type": "array"
+          },
+          "fares": {
+            "items": {
+              "$ref": "#/components/schemas/Tfl.Api.Presentation.Entities.JourneyPlanner.Fare"
+            },
+            "type": "array"
+          },
+          "totalCost": {
+            "format": "int32",
+            "type": "integer"
+          }
+        },
+        "type": "object"
+      },
+      "Tfl.Api.Presentation.Entities.JourneyPlanner.JourneyPlannerCycleHireDockingStationData": {
+        "properties": {
+          "destinationId": {
+            "type": "string"
+          },
+          "destinationNumberOfBikes": {
+            "format": "int32",
+            "type": "integer"
+          },
+          "destinationNumberOfEmptySlots": {
+            "format": "int32",
+            "type": "integer"
+          },
+          "originId": {
+            "type": "string"
+          },
+          "originNumberOfBikes": {
+            "format": "int32",
+            "type": "integer"
+          },
+          "originNumberOfEmptySlots": {
+            "format": "int32",
+            "type": "integer"
+          }
+        },
+        "type": "object"
+      },
+      "Tfl.Api.Presentation.Entities.JourneyPlanner.JourneyVector": {
+        "properties": {
+          "from": {
+            "type": "string"
+          },
+          "to": {
+            "type": "string"
+          },
+          "uri": {
+            "type": "string"
+          },
+          "via": {
+            "type": "string"
+          }
+        },
+        "type": "object"
+      },
+      "Tfl.Api.Presentation.Entities.JourneyPlanner.Leg": {
+        "properties": {
+          "arrivalPoint": {
+            "$ref": "#/components/schemas/Tfl.Api.Presentation.Entities.Point"
+          },
+          "arrivalTime": {
+            "format": "date-time",
+            "type": "string"
+          },
+          "departurePoint": {
+            "$ref": "#/components/schemas/Tfl.Api.Presentation.Entities.Point"
+          },
+          "departureTime": {
+            "format": "date-time",
+            "type": "string"
+          },
+          "disruptions": {
+            "items": {
+              "$ref": "#/components/schemas/Tfl.Api.Presentation.Entities.Disruption"
+            },
+            "type": "array"
+          },
+          "distance": {
+            "format": "double",
+            "type": "number"
+          },
+          "duration": {
+            "format": "int32",
+            "type": "integer"
+          },
+          "hasFixedLocations": {
+            "readOnly": true,
+            "type": "boolean"
+          },
+          "instruction": {
+            "$ref": "#/components/schemas/Tfl.Api.Presentation.Entities.Instruction"
+          },
+          "isDisrupted": {
+            "readOnly": true,
+            "type": "boolean"
+          },
+          "mode": {
+            "$ref": "#/components/schemas/Tfl.Api.Presentation.Entities.Identifier"
+          },
+          "obstacles": {
+            "items": {
+              "$ref": "#/components/schemas/Tfl.Api.Presentation.Entities.JourneyPlanner.Obstacle"
+            },
+            "type": "array"
+          },
+          "path": {
+            "$ref": "#/components/schemas/Tfl.Api.Presentation.Entities.JourneyPlanner.Path"
+          },
+          "plannedWorks": {
+            "items": {
+              "$ref": "#/components/schemas/Tfl.Api.Presentation.Entities.JourneyPlanner.PlannedWork"
+            },
+            "type": "array"
+          },
+          "routeOptions": {
+            "items": {
+              "$ref": "#/components/schemas/Tfl.Api.Presentation.Entities.JourneyPlanner.RouteOption"
+            },
+            "type": "array"
+          },
+          "speed": {
+            "type": "string"
+          }
+        },
+        "type": "object"
+      },
+      "Tfl.Api.Presentation.Entities.JourneyPlanner.Obstacle": {
+        "properties": {
+          "incline": {
+            "type": "string"
+          },
+          "position": {
+            "type": "string"
+          },
+          "stopId": {
+            "format": "int32",
+            "type": "integer"
+          },
+          "type": {
+            "type": "string"
+          }
+        },
+        "type": "object"
+      },
+      "Tfl.Api.Presentation.Entities.JourneyPlanner.Path": {
+        "properties": {
+          "elevation": {
+            "items": {
+              "$ref": "#/components/schemas/Tfl.Api.Common.JourneyPlanner.JpElevation"
+            },
+            "type": "array"
+          },
+          "lineString": {
+            "type": "string"
+          },
+          "stopPoints": {
+            "items": {
+              "$ref": "#/components/schemas/Tfl.Api.Presentation.Entities.Identifier"
+            },
+            "type": "array"
+          }
+        },
+        "type": "object"
+      },
+      "Tfl.Api.Presentation.Entities.JourneyPlanner.PlannedWork": {
+        "properties": {
+          "createdDateTime": {
+            "format": "date-time",
+            "type": "string"
+          },
+          "description": {
+            "type": "string"
+          },
+          "id": {
+            "type": "string"
+          },
+          "lastUpdateDateTime": {
+            "format": "date-time",
+            "type": "string"
+          }
+        },
+        "type": "object"
+      },
+      "Tfl.Api.Presentation.Entities.JourneyPlanner.RouteOption": {
+        "properties": {
+          "directions": {
+            "items": {
+              "type": "string"
+            },
+            "type": "array"
+          },
+          "id": {
+            "description": "The Id of the route",
+            "type": "string"
+          },
+          "lineIdentifier": {
+            "$ref": "#/components/schemas/Tfl.Api.Presentation.Entities.Identifier"
+          },
+          "name": {
+            "description": "Name such as \"72\"",
+            "type": "string"
+          }
+        },
+        "type": "object"
+      },
+      "Tfl.Api.Presentation.Entities.JourneyPlanner.SearchCriteria": {
+        "properties": {
+          "dateTime": {
+            "format": "date-time",
+            "type": "string"
+          },
+          "dateTimeType": {
+            "enum": [
+              "Arriving",
+              "Departing"
+            ],
+            "type": "string"
+          },
+          "timeAdjustments": {
+            "$ref": "#/components/schemas/Tfl.Api.Presentation.Entities.JourneyPlanner.TimeAdjustments"
+          }
+        },
+        "type": "object"
+      },
+      "Tfl.Api.Presentation.Entities.JourneyPlanner.TimeAdjustment": {
+        "properties": {
+          "date": {
+            "type": "string"
+          },
+          "time": {
+            "type": "string"
+          },
+          "timeIs": {
+            "type": "string"
+          },
+          "uri": {
+            "type": "string"
+          }
+        },
+        "type": "object"
+      },
+      "Tfl.Api.Presentation.Entities.JourneyPlanner.TimeAdjustments": {
+        "properties": {
+          "earlier": {
+            "$ref": "#/components/schemas/Tfl.Api.Presentation.Entities.JourneyPlanner.TimeAdjustment"
+          },
+          "earliest": {
+            "$ref": "#/components/schemas/Tfl.Api.Presentation.Entities.JourneyPlanner.TimeAdjustment"
+          },
+          "later": {
+            "$ref": "#/components/schemas/Tfl.Api.Presentation.Entities.JourneyPlanner.TimeAdjustment"
+          },
+          "latest": {
+            "$ref": "#/components/schemas/Tfl.Api.Presentation.Entities.JourneyPlanner.TimeAdjustment"
+          }
+        },
+        "type": "object"
+      },
+      "Tfl.Api.Presentation.Entities.KnownJourney": {
+        "properties": {
+          "hour": {
+            "type": "string"
+          },
+          "intervalId": {
+            "format": "int32",
+            "type": "integer"
+          },
+          "minute": {
+            "type": "string"
+          }
+        },
+        "type": "object"
+      },
+      "Tfl.Api.Presentation.Entities.Line": {
+        "properties": {
+          "created": {
+            "format": "date-time",
+            "type": "string"
+          },
+          "crowding": {
+            "$ref": "#/components/schemas/Tfl.Api.Presentation.Entities.Crowding"
+          },
+          "disruptions": {
+            "items": {
+              "$ref": "#/components/schemas/Tfl.Api.Presentation.Entities.Disruption"
+            },
+            "type": "array"
+          },
+          "id": {
+            "type": "string"
+          },
+          "lineStatuses": {
+            "items": {
+              "$ref": "#/components/schemas/Tfl.Api.Presentation.Entities.LineStatus"
+            },
+            "type": "array"
+          },
+          "modeName": {
+            "type": "string"
+          },
+          "modified": {
+            "format": "date-time",
+            "type": "string"
+          },
+          "name": {
+            "type": "string"
+          },
+          "routeSections": {
+            "items": {
+              "$ref": "#/components/schemas/Tfl.Api.Presentation.Entities.MatchedRoute"
+            },
+            "type": "array"
+          },
+          "serviceTypes": {
+            "items": {
+              "$ref": "#/components/schemas/Tfl.Api.Presentation.Entities.LineServiceTypeInfo"
+            },
+            "type": "array"
+          }
+        },
+        "type": "object"
+      },
+      "Tfl.Api.Presentation.Entities.LineGroup": {
+        "properties": {
+          "lineIdentifier": {
+            "items": {
+              "type": "string"
+            },
+            "type": "array"
+          },
+          "naptanIdReference": {
+            "type": "string"
+          },
+          "stationAtcoCode": {
+            "type": "string"
+          }
+        },
+        "type": "object"
+      },
+      "Tfl.Api.Presentation.Entities.LineModeGroup": {
+        "properties": {
+          "lineIdentifier": {
+            "items": {
+              "type": "string"
+            },
+            "type": "array"
+          },
+          "modeName": {
+            "type": "string"
+          }
+        },
+        "type": "object"
+      },
+      "Tfl.Api.Presentation.Entities.LineRouteSection": {
+        "properties": {
+          "destination": {
+            "type": "string"
+          },
+          "direction": {
+            "type": "string"
+          },
+          "fromStation": {
+            "type": "string"
+          },
+          "routeId": {
+            "format": "int32",
+            "type": "integer"
+          },
+          "serviceType": {
+            "type": "string"
+          },
+          "toStation": {
+            "type": "string"
+          },
+          "vehicleDestinationText": {
+            "type": "string"
+          }
+        },
+        "type": "object"
+      },
+      "Tfl.Api.Presentation.Entities.LineServiceType": {
+        "properties": {
+          "lineName": {
+            "type": "string"
+          },
+          "lineSpecificServiceTypes": {
+            "items": {
+              "$ref": "#/components/schemas/Tfl.Api.Presentation.Entities.LineSpecificServiceType"
+            },
+            "type": "array"
+          }
+        },
+        "type": "object"
+      },
+      "Tfl.Api.Presentation.Entities.LineServiceTypeInfo": {
+        "properties": {
+          "name": {
+            "type": "string"
+          },
+          "uri": {
+            "type": "string"
+          }
+        },
+        "type": "object"
+      },
+      "Tfl.Api.Presentation.Entities.LineSpecificServiceType": {
+        "properties": {
+          "serviceType": {
+            "$ref": "#/components/schemas/Tfl.Api.Presentation.Entities.LineServiceTypeInfo"
+          },
+          "stopServesServiceType": {
+            "type": "boolean"
+          }
+        },
+        "type": "object"
+      },
+      "Tfl.Api.Presentation.Entities.LineStatus": {
+        "properties": {
+          "created": {
+            "format": "date-time",
+            "type": "string"
+          },
+          "disruption": {
+            "$ref": "#/components/schemas/Tfl.Api.Presentation.Entities.Disruption"
+          },
+          "id": {
+            "format": "int32",
+            "type": "integer"
+          },
+          "lineId": {
+            "type": "string"
+          },
+          "modified": {
+            "format": "date-time",
+            "type": "string"
+          },
+          "reason": {
+            "type": "string"
+          },
+          "statusSeverity": {
+            "format": "int32",
+            "type": "integer"
+          },
+          "statusSeverityDescription": {
+            "type": "string"
+          },
+          "validityPeriods": {
+            "items": {
+              "$ref": "#/components/schemas/Tfl.Api.Presentation.Entities.ValidityPeriod"
+            },
+            "type": "array"
+          }
+        },
+        "type": "object"
+      },
+      "Tfl.Api.Presentation.Entities.MatchedRoute": {
+        "description": "Description of a Route used in Route search results.",
+        "properties": {
+          "destination": {
+            "description": "The Id (NaPTAN code) or the Destination StopPoint",
+            "type": "string"
+          },
+          "destinationName": {
+            "description": "The name of the Destination StopPoint",
+            "type": "string"
+          },
+          "direction": {
+            "description": "Inbound or Outbound",
+            "type": "string"
+          },
+          "name": {
+            "description": "Name such as \"72\"",
+            "type": "string"
+          },
+          "originationName": {
+            "description": "The name of the Origin StopPoint",
+            "type": "string"
+          },
+          "originator": {
+            "description": "The Id (NaPTAN code) of the Origin StopPoint",
+            "type": "string"
+          },
+          "routeCode": {
+            "description": "The route code",
+            "type": "string"
+          },
+          "serviceType": {
+            "description": "Regular or Night",
+            "type": "string"
+          },
+          "validFrom": {
+            "description": "The DateTime that the Service containing this Route is valid from.",
+            "format": "date-time",
+            "type": "string"
+          },
+          "validTo": {
+            "description": "The DateTime that the Service containing this Route is valid until.",
+            "format": "date-time",
+            "type": "string"
+          }
+        },
+        "type": "object"
+      },
+      "Tfl.Api.Presentation.Entities.MatchedRouteSections": {
+        "properties": {
+          "id": {
+            "format": "int32",
+            "type": "integer"
+          }
+        },
+        "type": "object"
+      },
+      "Tfl.Api.Presentation.Entities.MatchedStop": {
+        "properties": {
+          "accessibilitySummary": {
+            "type": "string"
+          },
+          "direction": {
+            "type": "string"
+          },
+          "hasDisruption": {
+            "type": "boolean"
+          },
+          "icsId": {
+            "type": "string"
+          },
+          "id": {
+            "type": "string"
+          },
+          "lat": {
+            "format": "double",
+            "type": "number"
+          },
+          "lines": {
+            "items": {
+              "$ref": "#/components/schemas/Tfl.Api.Presentation.Entities.Identifier"
+            },
+            "type": "array"
+          },
+          "lon": {
+            "format": "double",
+            "type": "number"
+          },
+          "modes": {
+            "items": {
+              "type": "string"
+            },
+            "type": "array"
+          },
+          "name": {
+            "type": "string"
+          },
+          "parentId": {
+            "type": "string"
+          },
+          "routeId": {
+            "format": "int32",
+            "type": "integer"
+          },
+          "stationId": {
+            "type": "string"
+          },
+          "status": {
+            "type": "boolean"
+          },
+          "stopLetter": {
+            "type": "string"
+          },
+          "stopType": {
+            "type": "string"
+          },
+          "topMostParentId": {
+            "type": "string"
+          },
+          "towards": {
+            "type": "string"
+          },
+          "url": {
+            "type": "string"
+          },
+          "zone": {
+            "type": "string"
+          }
+        },
+        "type": "object"
+      },
+      "Tfl.Api.Presentation.Entities.Message": {
+        "properties": {
+          "bulletOrder": {
+            "format": "int32",
+            "type": "integer"
+          },
+          "header": {
+            "type": "boolean"
+          },
+          "linkText": {
+            "type": "string"
+          },
+          "messageText": {
+            "type": "string"
+          },
+          "url": {
+            "type": "string"
+          }
+        },
+        "type": "object"
+      },
+      "Tfl.Api.Presentation.Entities.Mode": {
+        "properties": {
+          "isFarePaying": {
+            "type": "boolean"
+          },
+          "isScheduledService": {
+            "type": "boolean"
+          },
+          "isTflService": {
+            "type": "boolean"
+          },
+          "modeName": {
+            "type": "string"
+          }
+        },
+        "type": "object"
+      },
+      "Tfl.Api.Presentation.Entities.NetworkStatus": {
+        "description": "Represent travel network status",
+        "properties": {
+          "message": {
+            "type": "string"
+          },
+          "operator": {
+            "type": "string"
+          },
+          "status": {
+            "type": "string"
+          },
+          "statusLevel": {
+            "format": "int32",
+            "type": "integer"
+          }
+        },
+        "type": "object"
+      },
+      "Tfl.Api.Presentation.Entities.OrderedRoute": {
+        "properties": {
+          "name": {
+            "type": "string"
+          },
+          "naptanIds": {
+            "items": {
+              "type": "string"
+            },
+            "type": "array"
+          },
+          "serviceType": {
+            "type": "string"
+          }
+        },
+        "type": "object"
+      },
+      "Tfl.Api.Presentation.Entities.PassengerFlow": {
+        "properties": {
+          "timeSlice": {
+            "description": "Time in 24hr format with 15 minute intervals e.g. 0500-0515, 0515-0530 etc.",
+            "type": "string"
+          },
+          "value": {
+            "description": "Count of passenger flow towards a platform",
+            "format": "int32",
+            "type": "integer"
+          }
+        },
+        "type": "object"
+      },
+      "Tfl.Api.Presentation.Entities.PathAttribute": {
+        "properties": {
+          "name": {
+            "type": "string"
+          },
+          "value": {
+            "type": "string"
+          }
+        },
+        "type": "object"
+      },
+      "Tfl.Api.Presentation.Entities.Period": {
+        "properties": {
+          "frequency": {
+            "$ref": "#/components/schemas/Tfl.Api.Presentation.Entities.ServiceFrequency"
+          },
+          "fromTime": {
+            "$ref": "#/components/schemas/Tfl.Api.Presentation.Entities.TwentyFourHourClockTime"
+          },
+          "toTime": {
+            "$ref": "#/components/schemas/Tfl.Api.Presentation.Entities.TwentyFourHourClockTime"
+          },
+          "type": {
+            "enum": [
+              "Normal",
+              "FrequencyHours",
+              "FrequencyMinutes",
+              "Unknown"
+            ],
+            "type": "string"
+          }
+        },
+        "type": "object"
+      },
+      "Tfl.Api.Presentation.Entities.Place": {
+        "properties": {
+          "additionalProperties": {
+            "description": "A bag of additional key/value pairs with extra information about this place.",
+            "items": {
+              "$ref": "#/components/schemas/Tfl.Api.Presentation.Entities.AdditionalProperties"
+            },
+            "type": "array"
+          },
+          "children": {
+            "items": {
+              "$ref": "#/components/schemas/Tfl.Api.Presentation.Entities.Place"
+            },
+            "type": "array"
+          },
+          "childrenUrls": {
+            "items": {
+              "type": "string"
+            },
+            "type": "array"
+          },
+          "commonName": {
+            "description": "A human readable name.",
+            "type": "string"
+          },
+          "distance": {
+            "description": "The distance of the place from its search point, if this is the result\r\n            of a geographical search, otherwise zero.",
+            "format": "double",
+            "type": "number"
+          },
+          "id": {
+            "description": "A unique identifier.",
+            "type": "string"
+          },
+          "lat": {
+            "description": "WGS84 latitude of the location.",
+            "format": "double",
+            "type": "number"
+          },
+          "lon": {
+            "description": "WGS84 longitude of the location.",
+            "format": "double",
+            "type": "number"
+          },
+          "placeType": {
+            "description": "The type of Place. See /Place/Meta/placeTypes for possible values.",
+            "type": "string"
+          },
+          "url": {
+            "description": "The unique location of this resource.",
+            "type": "string"
+          }
+        },
+        "type": "object"
+      },
+      "Tfl.Api.Presentation.Entities.PlaceCategory": {
+        "properties": {
+          "availableKeys": {
+            "items": {
+              "type": "string"
+            },
+            "type": "array"
+          },
+          "category": {
+            "type": "string"
+          }
+        },
+        "type": "object"
+      },
+      "Tfl.Api.Presentation.Entities.PlacePolygon": {
+        "properties": {
+          "commonName": {
+            "type": "string"
+          },
+          "geoPoints": {
+            "items": {
+              "$ref": "#/components/schemas/Tfl.Api.Common.GeoPoint"
+            },
+            "type": "array"
+          }
+        },
+        "type": "object"
+      },
+      "Tfl.Api.Presentation.Entities.Point": {
+        "description": "Represents a point located at a latitude and longitude using the WGS84 co-ordinate system.",
+        "properties": {
+          "lat": {
+            "description": "WGS84 latitude of the location.",
+            "format": "double",
+            "type": "number"
+          },
+          "lon": {
+            "description": "WGS84 longitude of the location.",
+            "format": "double",
+            "type": "number"
+          }
+        },
+        "type": "object"
+      },
+      "Tfl.Api.Presentation.Entities.Prediction": {
+        "description": "DTO to capture the prediction details",
+        "properties": {
+          "bearing": {
+            "description": "Bearing (between 0 to 359)",
+            "type": "string"
+          },
+          "currentLocation": {
+            "description": "The current location of the vehicle.",
+            "type": "string"
+          },
+          "destinationName": {
+            "description": "Name of the destination",
+            "type": "string"
+          },
+          "destinationNaptanId": {
+            "description": "Naptan Identifier for the prediction's destination",
+            "type": "string"
+          },
+          "direction": {
+            "description": "Direction (unified to inbound/outbound)",
+            "type": "string"
+          },
+          "expectedArrival": {
+            "description": "The expected arrival time of the vehicle at the stop/station",
+            "format": "date-time",
+            "type": "string"
+          },
+          "id": {
+            "description": "The identitier for the prediction",
+            "type": "string"
+          },
+          "lineId": {
+            "description": "Unique identifier for the Line",
+            "type": "string"
+          },
+          "lineName": {
+            "description": "Line Name",
+            "type": "string"
+          },
+          "modeName": {
+            "description": "The mode name of the station/line the prediction relates to",
+            "type": "string"
+          },
+          "naptanId": {
+            "description": "Identifier for the prediction",
+            "type": "string"
+          },
+          "operationType": {
+            "description": "The type of the operation (1: is new or has been updated, 2: should be deleted from any client cache)",
+            "format": "int32",
+            "type": "integer"
+          },
+          "platformName": {
+            "description": "Platform name (for bus, this is the stop letter)",
+            "type": "string"
+          },
+          "stationName": {
+            "description": "Station name",
+            "type": "string"
+          },
+          "timeToLive": {
+            "description": "The expiry time for the prediction",
+            "format": "date-time",
+            "type": "string"
+          },
+          "timeToStation": {
+            "description": "Prediction of the Time to station in seconds",
+            "format": "int32",
+            "type": "integer"
+          },
+          "timestamp": {
+            "description": "Timestamp for when the prediction was inserted/modified (source column drives what objects are broadcast on each iteration)",
+            "format": "date-time",
+            "type": "string"
+          },
+          "timing": {
+            "$ref": "#/components/schemas/Tfl.Api.Presentation.Entities.PredictionTiming"
+          },
+          "towards": {
+            "description": "Routing information or other descriptive text about the path of the vehicle towards the destination",
+            "type": "string"
+          },
+          "vehicleId": {
+            "description": "The actual vehicle in transit (for train modes, the leading car of the rolling set)",
+            "type": "string"
+          }
+        },
+        "type": "object"
+      },
+      "Tfl.Api.Presentation.Entities.PredictionTiming": {
+        "properties": {
+          "countdownServerAdjustment": {
+            "type": "string"
+          },
+          "insert": {
+            "format": "date-time",
+            "type": "string"
+          },
+          "read": {
+            "format": "date-time",
+            "type": "string"
+          },
+          "received": {
+            "format": "date-time",
+            "type": "string"
+          },
+          "sent": {
+            "format": "date-time",
+            "type": "string"
+          },
+          "source": {
+            "format": "date-time",
+            "type": "string"
+          }
+        },
+        "type": "object"
+      },
+      "Tfl.Api.Presentation.Entities.Redirect": {
+        "properties": {
+          "active": {
+            "type": "boolean"
+          },
+          "longUrl": {
+            "type": "string"
+          },
+          "shortUrl": {
+            "type": "string"
+          }
+        },
+        "type": "object"
+      },
+      "Tfl.Api.Presentation.Entities.RoadCorridor": {
+        "properties": {
+          "bounds": {
+            "description": "The Bounds of the Corridor, given by the south-east followed by the north-west co-ordinate\r\n            pair in geoJSON format e.g. \"[[-1.241531,51.242151],[1.641223,53.765721]]\"",
+            "type": "string"
+          },
+          "displayName": {
+            "description": "The display name of the Corridor e.g. \"North Circular (A406)\". This\r\n            may be identical to the Id.",
+            "type": "string"
+          },
+          "envelope": {
+            "description": "The Envelope of the Corridor, given by the corner co-ordinates of a rectangular (four-point) polygon\r\n            in geoJSON format e.g. \"[[-1.241531,51.242151],[-1.241531,53.765721],[1.641223,53.765721],[1.641223,51.242151]]\"",
+            "type": "string"
+          },
+          "group": {
+            "description": "The group name of the Corridor e.g. \"Central London\". Most corridors are not grouped, in which case this field can be null.",
+            "type": "string"
+          },
+          "id": {
+            "description": "The Id of the Corridor e.g. \"A406\"",
+            "type": "string"
+          },
+          "statusAggregationEndDate": {
+            "description": "The end of the period over which status has been aggregated, or null if this is the current corridor status.",
+            "format": "date-time",
+            "type": "string"
+          },
+          "statusAggregationStartDate": {
+            "description": "The start of the period over which status has been aggregated, or null if this is the current corridor status.",
+            "format": "date-time",
+            "type": "string"
+          },
+          "statusSeverity": {
+            "description": "Standard multi-mode status severity code",
+            "type": "string"
+          },
+          "statusSeverityDescription": {
+            "description": "Description of the status severity as applied to RoadCorridors",
+            "type": "string"
+          },
+          "url": {
+            "description": "URL to retrieve this Corridor.",
+            "type": "string"
+          }
+        },
+        "type": "object"
+      },
+      "Tfl.Api.Presentation.Entities.RoadDisruption": {
+        "properties": {
+          "category": {
+            "description": "Describes the nature of disruption e.g. Traffic Incidents, Works",
+            "type": "string"
+          },
+          "comments": {
+            "description": "Full text of comments describing the disruption, including details of any road closures and diversions, where appropriate.",
+            "type": "string"
+          },
+          "corridorIds": {
+            "description": "The Ids of affected corridors, if any.",
+            "items": {
+              "type": "string"
+            },
+            "type": "array"
+          },
+          "currentUpdate": {
+            "description": "Text of the most recent update from the LSTCC on the state of the\r\n             disruption, including the current traffic impact and any advice to\r\n             road users.",
+            "type": "string"
+          },
+          "currentUpdateDateTime": {
+            "description": "The time when the last CurrentUpdate description was recorded, \r\n            or null if no CurrentUpdate has been applied.",
+            "format": "date-time",
+            "type": "string"
+          },
+          "endDateTime": {
+            "description": "The date and time on which the disruption ended. For planned disruptions, this date will have a valid value. For unplanned \r\n            disruptions in progress, this field will be omitted.",
+            "format": "date-time",
+            "type": "string"
+          },
+          "geography": {
+            "$ref": "#/components/schemas/System.Data.Spatial.DbGeography"
+          },
+          "geometry": {
+            "$ref": "#/components/schemas/System.Data.Spatial.DbGeography"
+          },
+          "hasClosures": {
+            "description": "True if any of the affected Streets have a \"Full Closure\" status, false otherwise. A RoadDisruption that has HasClosures is considered a \r\n            Severe or Serious disruption for severity filtering purposes.",
+            "type": "boolean"
+          },
+          "id": {
+            "description": "Unique identifier for the road disruption",
+            "type": "string"
+          },
+          "isProvisional": {
+            "description": "True if the disruption is planned on a future date that is open to change",
+            "type": "boolean"
+          },
+          "lastModifiedTime": {
+            "description": "The date and time on which the disruption was last modified in the system. This information can reliably be used by a developer to quickly\r\n            compare two instances of the same disruption to determine if it has been changed.",
+            "format": "date-time",
+            "type": "string"
+          },
+          "levelOfInterest": {
+            "description": "This describes the level of potential impact on traffic operations of the disruption. \r\n            High = e.g. a one-off disruption on a major or high profile route which will require a high level of operational attention \r\n            Medium = This is the default value \r\n            Low = e.g. a frequently occurring disruption which is well known",
+            "type": "string"
+          },
+          "linkText": {
+            "description": "The text of any associated link",
+            "type": "string"
+          },
+          "linkUrl": {
+            "description": "The url of any associated link",
+            "type": "string"
+          },
+          "location": {
+            "description": "Main road name / number (borough) or preset area name where the disruption is located. This might be useful for a map popup where space is limited.",
+            "type": "string"
+          },
+          "ordinal": {
+            "description": "An ordinal of the disruption based on severity, level of interest and corridor.",
+            "format": "int32",
+            "type": "integer"
+          },
+          "point": {
+            "description": "Latitude and longitude (WGS84) of the centroid of the disruption, stored in a geoJSON-formatted string.",
+            "type": "string"
+          },
+          "publishEndDate": {
+            "format": "date-time",
+            "type": "string"
+          },
+          "publishStartDate": {
+            "description": "TDM Additional properties",
+            "format": "date-time",
+            "type": "string"
+          },
+          "recurringSchedules": {
+            "items": {
+              "$ref": "#/components/schemas/Tfl.Api.Presentation.Entities.RoadDisruptionSchedule"
+            },
+            "type": "array"
+          },
+          "roadDisruptionImpactAreas": {
+            "items": {
+              "$ref": "#/components/schemas/Tfl.Api.Presentation.Entities.RoadDisruptionImpactArea"
+            },
+            "type": "array"
+          },
+          "roadDisruptionLines": {
+            "items": {
+              "$ref": "#/components/schemas/Tfl.Api.Presentation.Entities.RoadDisruptionLine"
+            },
+            "type": "array"
+          },
+          "roadProject": {
+            "$ref": "#/components/schemas/Tfl.Api.Presentation.Entities.RoadProject"
+          },
+          "severity": {
+            "description": "A description of the severity of the disruption.",
+            "type": "string"
+          },
+          "startDateTime": {
+            "description": "The date and time which the disruption started. For a planned disruption (i.e. planned road works) this date will be in the future.\r\n            For unplanned disruptions, this will default to the date on which the disruption was first recorded, but may be adjusted by the operator.",
+            "format": "date-time",
+            "type": "string"
+          },
+          "status": {
+            "description": "This describes the status of the disruption.  \r\n            Active = currently in progress \r\n            Active Long Term = currently in progress and long term\r\n            Scheduled = scheduled to start within the next 180 days\r\n            Recurring Works = planned maintenance works that follow a regular routine or pattern and whose next occurrence is to start within the next 180 days.\r\n            Recently Cleared = recently cleared in the last 24 hours\r\n            Note that the status of Scheduled or Recurring Works disruptions will change to Active when they start, and will change status again when they end.",
+            "type": "string"
+          },
+          "streets": {
+            "description": "A collection of zero or more streets affected by the disruption.",
+            "items": {
+              "$ref": "#/components/schemas/Tfl.Api.Presentation.Entities.Street"
+            },
+            "type": "array"
+          },
+          "subCategory": {
+            "description": "Describes the sub-category of disruption e.g. Collapsed Manhole, Abnormal Load",
+            "type": "string"
+          },
+          "timeFrame": {
+            "type": "string"
+          },
+          "url": {
+            "description": "URL to retrieve this road disruption",
+            "type": "string"
+          }
+        },
+        "type": "object"
+      },
+      "Tfl.Api.Presentation.Entities.RoadDisruptionImpactArea": {
+        "properties": {
+          "endDate": {
+            "format": "date-time",
+            "type": "string"
+          },
+          "endTime": {
+            "type": "string"
+          },
+          "id": {
+            "format": "int32",
+            "type": "integer"
+          },
+          "polygon": {
+            "$ref": "#/components/schemas/System.Data.Spatial.DbGeography"
+          },
+          "roadDisruptionId": {
+            "type": "string"
+          },
+          "startDate": {
+            "format": "date-time",
+            "type": "string"
+          },
+          "startTime": {
+            "type": "string"
+          }
+        },
+        "type": "object"
+      },
+      "Tfl.Api.Presentation.Entities.RoadDisruptionLine": {
+        "properties": {
+          "endDate": {
+            "format": "date-time",
+            "type": "string"
+          },
+          "endTime": {
+            "type": "string"
+          },
+          "id": {
+            "format": "int32",
+            "type": "integer"
+          },
+          "isDiversion": {
+            "type": "boolean"
+          },
+          "multiLineString": {
+            "$ref": "#/components/schemas/System.Data.Spatial.DbGeography"
+          },
+          "roadDisruptionId": {
+            "type": "string"
+          },
+          "startDate": {
+            "format": "date-time",
+            "type": "string"
+          },
+          "startTime": {
+            "type": "string"
+          }
+        },
+        "type": "object"
+      },
+      "Tfl.Api.Presentation.Entities.RoadDisruptionSchedule": {
+        "properties": {
+          "endTime": {
+            "format": "date-time",
+            "type": "string"
+          },
+          "startTime": {
+            "format": "date-time",
+            "type": "string"
+          }
+        },
+        "type": "object"
+      },
+      "Tfl.Api.Presentation.Entities.RoadProject": {
+        "properties": {
+          "boroughsBenefited": {
+            "items": {
+              "type": "string"
+            },
+            "type": "array"
+          },
+          "constructionEndDate": {
+            "format": "date-time",
+            "type": "string"
+          },
+          "constructionStartDate": {
+            "format": "date-time",
+            "type": "string"
+          },
+          "consultationEndDate": {
+            "format": "date-time",
+            "type": "string"
+          },
+          "consultationPageUrl": {
+            "type": "string"
+          },
+          "consultationStartDate": {
+            "format": "date-time",
+            "type": "string"
+          },
+          "contactEmail": {
+            "type": "string"
+          },
+          "contactName": {
+            "type": "string"
+          },
+          "cycleSuperhighwayId": {
+            "type": "string"
+          },
+          "externalPageUrl": {
+            "type": "string"
+          },
+          "phase": {
+            "enum": [
+              "Unscoped",
+              "Concept",
+              "ConsultationEnded",
+              "Consultation",
+              "Construction",
+              "Complete"
+            ],
+            "type": "string"
+          },
+          "projectDescription": {
+            "type": "string"
+          },
+          "projectId": {
+            "type": "string"
+          },
+          "projectName": {
+            "type": "string"
+          },
+          "projectPageUrl": {
+            "type": "string"
+          },
+          "projectSummaryPageUrl": {
+            "type": "string"
+          },
+          "schemeName": {
+            "type": "string"
+          }
+        },
+        "type": "object"
+      },
+      "Tfl.Api.Presentation.Entities.RouteSearchMatch": {
+        "properties": {
+          "id": {
+            "type": "string"
+          },
+          "lat": {
+            "format": "double",
+            "type": "number"
+          },
+          "lineId": {
+            "type": "string"
+          },
+          "lineName": {
+            "type": "string"
+          },
+          "lineRouteSection": {
+            "items": {
+              "$ref": "#/components/schemas/Tfl.Api.Presentation.Entities.LineRouteSection"
+            },
+            "type": "array"
+          },
+          "lon": {
+            "format": "double",
+            "type": "number"
+          },
+          "matchedRouteSections": {
+            "items": {
+              "$ref": "#/components/schemas/Tfl.Api.Presentation.Entities.MatchedRouteSections"
+            },
+            "type": "array"
+          },
+          "matchedStops": {
+            "items": {
+              "$ref": "#/components/schemas/Tfl.Api.Presentation.Entities.MatchedStop"
+            },
+            "type": "array"
+          },
+          "mode": {
+            "type": "string"
+          },
+          "name": {
+            "type": "string"
+          },
+          "url": {
+            "type": "string"
+          }
+        },
+        "type": "object"
+      },
+      "Tfl.Api.Presentation.Entities.RouteSearchResponse": {
+        "properties": {
+          "input": {
+            "type": "string"
+          },
+          "searchMatches": {
+            "items": {
+              "$ref": "#/components/schemas/Tfl.Api.Presentation.Entities.RouteSearchMatch"
+            },
+            "type": "array"
+          }
+        },
+        "type": "object"
+      },
+      "Tfl.Api.Presentation.Entities.RouteSection": {
+        "properties": {
+          "destinationName": {
+            "description": "The name of the Destination StopPoint",
+            "type": "string"
+          },
+          "direction": {
+            "description": "Inbound or Outbound",
+            "type": "string"
+          },
+          "id": {
+            "description": "The Id of the route",
+            "type": "string"
+          },
+          "lineId": {
+            "description": "The Id of the Line",
+            "type": "string"
+          },
+          "lineString": {
+            "description": "The co-ordinates of the route's path as a geoJSON lineString",
+            "type": "string"
+          },
+          "name": {
+            "description": "Name such as \"72\"",
+            "type": "string"
+          },
+          "originationName": {
+            "description": "The name of the Origin StopPoint",
+            "type": "string"
+          },
+          "routeCode": {
+            "description": "The route code",
+            "type": "string"
+          },
+          "routeSectionNaptanEntrySequence": {
+            "items": {
+              "$ref": "#/components/schemas/Tfl.Api.Presentation.Entities.RouteSectionNaptanEntrySequence"
+            },
+            "type": "array"
+          },
+          "validFrom": {
+            "description": "The DateTime that the Service containing this Route is valid from.",
+            "format": "date-time",
+            "type": "string"
+          },
+          "validTo": {
+            "description": "The DateTime that the Service containing this Route is valid until.",
+            "format": "date-time",
+            "type": "string"
+          }
+        },
+        "type": "object"
+      },
+      "Tfl.Api.Presentation.Entities.RouteSectionNaptanEntrySequence": {
+        "properties": {
+          "ordinal": {
+            "format": "int32",
+            "type": "integer"
+          },
+          "stopPoint": {
+            "$ref": "#/components/schemas/Tfl.Api.Presentation.Entities.StopPoint"
+          }
+        },
+        "type": "object"
+      },
+      "Tfl.Api.Presentation.Entities.RouteSequence": {
+        "properties": {
+          "direction": {
+            "type": "string"
+          },
+          "isOutboundOnly": {
+            "type": "boolean"
+          },
+          "lineId": {
+            "type": "string"
+          },
+          "lineName": {
+            "type": "string"
+          },
+          "lineStrings": {
+            "items": {
+              "type": "string"
+            },
+            "type": "array"
+          },
+          "mode": {
+            "type": "string"
+          },
+          "orderedLineRoutes": {
+            "items": {
+              "$ref": "#/components/schemas/Tfl.Api.Presentation.Entities.OrderedRoute"
+            },
+            "type": "array"
+          },
+          "stations": {
+            "items": {
+              "$ref": "#/components/schemas/Tfl.Api.Presentation.Entities.MatchedStop"
+            },
+            "type": "array"
+          },
+          "stopPointSequences": {
+            "items": {
+              "$ref": "#/components/schemas/Tfl.Api.Presentation.Entities.StopPointSequence"
+            },
+            "type": "array"
+          }
+        },
+        "type": "object"
+      },
+      "Tfl.Api.Presentation.Entities.Schedule": {
+        "properties": {
+          "firstJourney": {
+            "$ref": "#/components/schemas/Tfl.Api.Presentation.Entities.KnownJourney"
+          },
+          "knownJourneys": {
+            "items": {
+              "$ref": "#/components/schemas/Tfl.Api.Presentation.Entities.KnownJourney"
+            },
+            "type": "array"
+          },
+          "lastJourney": {
+            "$ref": "#/components/schemas/Tfl.Api.Presentation.Entities.KnownJourney"
+          },
+          "name": {
+            "type": "string"
+          },
+          "periods": {
+            "items": {
+              "$ref": "#/components/schemas/Tfl.Api.Presentation.Entities.Period"
+            },
+            "type": "array"
+          }
+        },
+        "type": "object"
+      },
+      "Tfl.Api.Presentation.Entities.SearchMatch": {
+        "properties": {
+          "id": {
+            "type": "string"
+          },
+          "lat": {
+            "format": "double",
+            "type": "number"
+          },
+          "lon": {
+            "format": "double",
+            "type": "number"
+          },
+          "name": {
+            "type": "string"
+          },
+          "url": {
+            "type": "string"
+          }
+        },
+        "type": "object"
+      },
+      "Tfl.Api.Presentation.Entities.SearchResponse": {
+        "properties": {
+          "from": {
+            "format": "int32",
+            "type": "integer"
+          },
+          "matches": {
+            "items": {
+              "$ref": "#/components/schemas/Tfl.Api.Presentation.Entities.SearchMatch"
+            },
+            "type": "array"
+          },
+          "maxScore": {
+            "format": "double",
+            "type": "number"
+          },
+          "page": {
+            "format": "int32",
+            "type": "integer"
+          },
+          "pageSize": {
+            "format": "int32",
+            "type": "integer"
+          },
+          "provider": {
+            "type": "string"
+          },
+          "query": {
+            "type": "string"
+          },
+          "total": {
+            "format": "int32",
+            "type": "integer"
+          }
+        },
+        "type": "object"
+      },
+      "Tfl.Api.Presentation.Entities.ServiceFrequency": {
+        "properties": {
+          "highestFrequency": {
+            "format": "double",
+            "type": "number"
+          },
+          "lowestFrequency": {
+            "format": "double",
+            "type": "number"
+          }
+        },
+        "type": "object"
+      },
+      "Tfl.Api.Presentation.Entities.StationInterval": {
+        "properties": {
+          "id": {
+            "type": "string"
+          },
+          "intervals": {
+            "items": {
+              "$ref": "#/components/schemas/Tfl.Api.Presentation.Entities.Interval"
+            },
+            "type": "array"
+          }
+        },
+        "type": "object"
+      },
+      "Tfl.Api.Presentation.Entities.StatusSeverity": {
+        "properties": {
+          "description": {
+            "type": "string"
+          },
+          "modeName": {
+            "type": "string"
+          },
+          "severityLevel": {
+            "format": "int32",
+            "type": "integer"
+          }
+        },
+        "type": "object"
+      },
+      "Tfl.Api.Presentation.Entities.StopPoint": {
+        "properties": {
+          "accessibilitySummary": {
+            "type": "string"
+          },
+          "additionalProperties": {
+            "description": "A bag of additional key/value pairs with extra information about this place.",
+            "items": {
+              "$ref": "#/components/schemas/Tfl.Api.Presentation.Entities.AdditionalProperties"
+            },
+            "type": "array"
+          },
+          "children": {
+            "items": {
+              "$ref": "#/components/schemas/Tfl.Api.Presentation.Entities.Place"
+            },
+            "type": "array"
+          },
+          "childrenUrls": {
+            "items": {
+              "type": "string"
+            },
+            "type": "array"
+          },
+          "commonName": {
+            "description": "A human readable name.",
+            "type": "string"
+          },
+          "distance": {
+            "description": "The distance of the place from its search point, if this is the result\r\n            of a geographical search, otherwise zero.",
+            "format": "double",
+            "type": "number"
+          },
+          "fullName": {
+            "type": "string"
+          },
+          "hubNaptanCode": {
+            "type": "string"
+          },
+          "icsCode": {
+            "type": "string"
+          },
+          "id": {
+            "description": "A unique identifier.",
+            "type": "string"
+          },
+          "indicator": {
+            "description": "The indicator of the stop point e.g. \"Stop K\"",
+            "type": "string"
+          },
+          "lat": {
+            "description": "WGS84 latitude of the location.",
+            "format": "double",
+            "type": "number"
+          },
+          "lineGroup": {
+            "items": {
+              "$ref": "#/components/schemas/Tfl.Api.Presentation.Entities.LineGroup"
+            },
+            "type": "array"
+          },
+          "lineModeGroups": {
+            "items": {
+              "$ref": "#/components/schemas/Tfl.Api.Presentation.Entities.LineModeGroup"
+            },
+            "type": "array"
+          },
+          "lines": {
+            "items": {
+              "$ref": "#/components/schemas/Tfl.Api.Presentation.Entities.Identifier"
+            },
+            "type": "array"
+          },
+          "lon": {
+            "description": "WGS84 longitude of the location.",
+            "format": "double",
+            "type": "number"
+          },
+          "modes": {
+            "items": {
+              "type": "string"
+            },
+            "type": "array"
+          },
+          "naptanId": {
+            "type": "string"
+          },
+          "naptanMode": {
+            "type": "string"
+          },
+          "placeType": {
+            "description": "The type of Place. See /Place/Meta/placeTypes for possible values.",
+            "type": "string"
+          },
+          "platformName": {
+            "type": "string"
+          },
+          "smsCode": {
+            "type": "string"
+          },
+          "stationNaptan": {
+            "type": "string"
+          },
+          "status": {
+            "type": "boolean"
+          },
+          "stopLetter": {
+            "description": "The stop letter, if it could be cleansed from the Indicator e.g. \"K\"",
+            "type": "string"
+          },
+          "stopType": {
+            "type": "string"
+          },
+          "url": {
+            "description": "The unique location of this resource.",
+            "type": "string"
+          }
+        },
+        "type": "object"
+      },
+      "Tfl.Api.Presentation.Entities.StopPointCategory": {
+        "properties": {
+          "availableKeys": {
+            "items": {
+              "type": "string"
+            },
+            "type": "array"
+          },
+          "category": {
+            "type": "string"
+          }
+        },
+        "type": "object"
+      },
+      "Tfl.Api.Presentation.Entities.StopPointRouteSection": {
+        "properties": {
+          "destinationName": {
+            "type": "string"
+          },
+          "direction": {
+            "type": "string"
+          },
+          "isActive": {
+            "type": "boolean"
+          },
+          "lineId": {
+            "type": "string"
+          },
+          "lineString": {
+            "type": "string"
+          },
+          "mode": {
+            "type": "string"
+          },
+          "naptanId": {
+            "type": "string"
+          },
+          "routeSectionName": {
+            "type": "string"
+          },
+          "serviceType": {
+            "type": "string"
+          },
+          "validFrom": {
+            "format": "date-time",
+            "type": "string"
+          },
+          "validTo": {
+            "format": "date-time",
+            "type": "string"
+          },
+          "vehicleDestinationText": {
+            "type": "string"
+          }
+        },
+        "type": "object"
+      },
+      "Tfl.Api.Presentation.Entities.StopPointSequence": {
+        "properties": {
+          "branchId": {
+            "description": "The id of this branch.",
+            "format": "int32",
+            "type": "integer"
+          },
+          "direction": {
+            "type": "string"
+          },
+          "lineId": {
+            "type": "string"
+          },
+          "lineName": {
+            "type": "string"
+          },
+          "nextBranchIds": {
+            "description": "The ids of the next branch(es) in the sequence. Note that the next and previous branch id can be\r\n            identical in the case of a looped route e.g. the Circle line.",
+            "items": {
+              "format": "int32",
+              "type": "integer"
+            },
+            "type": "array"
+          },
+          "prevBranchIds": {
+            "description": "The ids of the previous branch(es) in the sequence. Note that the next and previous branch id can be\r\n            identical in the case of a looped route e.g. the Circle line.",
+            "items": {
+              "format": "int32",
+              "type": "integer"
+            },
+            "type": "array"
+          },
+          "serviceType": {
+            "enum": [
+              "Regular",
+              "Night"
+            ],
+            "type": "string"
+          },
+          "stopPoint": {
+            "items": {
+              "$ref": "#/components/schemas/Tfl.Api.Presentation.Entities.MatchedStop"
+            },
+            "type": "array"
+          }
+        },
+        "type": "object"
+      },
+      "Tfl.Api.Presentation.Entities.StopPointsResponse": {
+        "description": "A paged response containing StopPoints",
+        "properties": {
+          "centrePoint": {
+            "description": "The centre latitude/longitude of this list of StopPoints",
+            "items": {
+              "format": "double",
+              "type": "number"
+            },
+            "type": "array"
+          },
+          "page": {
+            "description": "The index of this page",
+            "format": "int32",
+            "type": "integer"
+          },
+          "pageSize": {
+            "description": "The maximum size of the page in this response i.e. the maximum number of StopPoints",
+            "format": "int32",
+            "type": "integer"
+          },
+          "stopPoints": {
+            "description": "Collection of stop points",
+            "items": {
+              "$ref": "#/components/schemas/Tfl.Api.Presentation.Entities.StopPoint"
+            },
+            "type": "array"
+          },
+          "total": {
+            "description": "The total number of StopPoints available across all pages",
+            "format": "int32",
+            "type": "integer"
+          }
+        },
+        "type": "object"
+      },
+      "Tfl.Api.Presentation.Entities.Street": {
+        "properties": {
+          "closure": {
+            "description": "Type of road closure. Some example values:\r\n            Open = road is open, not blocked, not closed, not restricted. It maybe that the disruption has been moved out of the carriageway.\r\n            Partial Closure = road is partially blocked, closed or restricted. \r\n            Full Closure = road is fully blocked or closed.",
+            "type": "string"
+          },
+          "directions": {
+            "description": "The direction of the disruption on the street. Some example values:\r\n            All Directions\r\n            All Approaches\r\n            Clockwise\r\n            Anti-Clockwise\r\n            Northbound\r\n            Eastbound\r\n            Southbound\r\n            Westbound\r\n            Both Directions",
+            "type": "string"
+          },
+          "name": {
+            "description": "Street name",
+            "type": "string"
+          },
+          "segments": {
+            "description": "Geographic description of the sections of this street that are affected.",
+            "items": {
+              "$ref": "#/components/schemas/Tfl.Api.Presentation.Entities.StreetSegment"
+            },
+            "type": "array"
+          },
+          "sourceSystemId": {
+            "description": "The ID from the source system of the disruption that this street belongs to.",
+            "format": "int64",
+            "type": "integer"
+          },
+          "sourceSystemKey": {
+            "description": "The key of the source system of the disruption that this street belongs to.",
+            "type": "string"
+          }
+        },
+        "type": "object"
+      },
+      "Tfl.Api.Presentation.Entities.StreetSegment": {
+        "properties": {
+          "lineString": {
+            "description": "geoJSON formatted LineString containing two latitude/longitude (WGS84) pairs that identify the start and end points of the street segment.",
+            "type": "string"
+          },
+          "sourceSystemId": {
+            "description": "The ID from the source system of the disruption that this street belongs to.",
+            "format": "int64",
+            "type": "integer"
+          },
+          "sourceSystemKey": {
+            "description": "The key of the source system of the disruption that this street belongs to.",
+            "type": "string"
+          },
+          "toid": {
+            "description": "A 16 digit unique integer identifying a OS ITN (Ordnance Survey Integrated Transport Network) road link.",
+            "type": "string"
+          }
+        },
+        "type": "object"
+      },
+      "Tfl.Api.Presentation.Entities.Timetable": {
+        "properties": {
+          "departureStopId": {
+            "type": "string"
+          },
+          "routes": {
+            "items": {
+              "$ref": "#/components/schemas/Tfl.Api.Presentation.Entities.TimetableRoute"
+            },
+            "type": "array"
+          }
+        },
+        "type": "object"
+      },
+      "Tfl.Api.Presentation.Entities.TimetableResponse": {
+        "properties": {
+          "direction": {
+            "type": "string"
+          },
+          "disambiguation": {
+            "$ref": "#/components/schemas/Tfl.Api.Presentation.Entities.Timetables.Disambiguation"
+          },
+          "lineId": {
+            "type": "string"
+          },
+          "lineName": {
+            "type": "string"
+          },
+          "pdfUrl": {
+            "type": "string"
+          },
+          "stations": {
+            "items": {
+              "$ref": "#/components/schemas/Tfl.Api.Presentation.Entities.MatchedStop"
+            },
+            "type": "array"
+          },
+          "statusErrorMessage": {
+            "type": "string"
+          },
+          "stops": {
+            "items": {
+              "$ref": "#/components/schemas/Tfl.Api.Presentation.Entities.MatchedStop"
+            },
+            "type": "array"
+          },
+          "timetable": {
+            "$ref": "#/components/schemas/Tfl.Api.Presentation.Entities.Timetable"
+          }
+        },
+        "type": "object"
+      },
+      "Tfl.Api.Presentation.Entities.TimetableRoute": {
+        "properties": {
+          "schedules": {
+            "items": {
+              "$ref": "#/components/schemas/Tfl.Api.Presentation.Entities.Schedule"
+            },
+            "type": "array"
+          },
+          "stationIntervals": {
+            "items": {
+              "$ref": "#/components/schemas/Tfl.Api.Presentation.Entities.StationInterval"
+            },
+            "type": "array"
+          }
+        },
+        "type": "object"
+      },
+      "Tfl.Api.Presentation.Entities.Timetables.Disambiguation": {
+        "properties": {
+          "disambiguationOptions": {
+            "items": {
+              "$ref": "#/components/schemas/Tfl.Api.Presentation.Entities.Timetables.DisambiguationOption"
+            },
+            "type": "array"
+          }
+        },
+        "type": "object"
+      },
+      "Tfl.Api.Presentation.Entities.Timetables.DisambiguationOption": {
+        "properties": {
+          "description": {
+            "type": "string"
+          },
+          "uri": {
+            "type": "string"
+          }
+        },
+        "type": "object"
+      },
+      "Tfl.Api.Presentation.Entities.TrainLoading": {
+        "properties": {
+          "direction": {
+            "description": "Direction in regards to Journey Planner i.e. inbound or outbound",
+            "type": "string"
+          },
+          "line": {
+            "description": "The Line Name e.g. \"Victoria\"",
+            "type": "string"
+          },
+          "lineDirection": {
+            "description": "Direction of the Line e.g. NB, SB, WB etc.",
+            "type": "string"
+          },
+          "naptanTo": {
+            "description": "Naptan of the adjacent station",
+            "type": "string"
+          },
+          "platformDirection": {
+            "description": "Direction displayed on the platform e.g. NB, SB, WB etc.",
+            "type": "string"
+          },
+          "timeSlice": {
+            "description": "Time in 24hr format with 15 minute intervals e.g. 0500-0515, 0515-0530 etc.",
+            "type": "string"
+          },
+          "value": {
+            "description": "Scale between 1-6, \r\n             1 = Very quiet, 2 = Quiet, 3 = Fairly busy, 4 = Busy, 5 = Very busy, 6 = Exceptionally busy",
+            "format": "int32",
+            "type": "integer"
+          }
+        },
+        "type": "object"
+      },
+      "Tfl.Api.Presentation.Entities.TwentyFourHourClockTime": {
+        "properties": {
+          "hour": {
+            "type": "string"
+          },
+          "minute": {
+            "type": "string"
+          }
+        },
+        "type": "object"
+      },
+      "Tfl.Api.Presentation.Entities.ValidityPeriod": {
+        "description": "Represents a period for which a planned works is valid.",
+        "properties": {
+          "fromDate": {
+            "description": "Gets or sets the start date.",
+            "format": "date-time",
+            "type": "string"
+          },
+          "isNow": {
+            "description": "If true is a realtime status rather than planned or info",
+            "type": "boolean"
+          },
+          "toDate": {
+            "description": "Gets or sets the end date.",
+            "format": "date-time",
+            "type": "string"
+          }
+        },
+        "type": "object"
+      },
+      "Tfl.Api.Presentation.Entities.VehicleMatch": {
+        "description": "",
+        "properties": {
+          "colour": {
+            "type": "string"
+          },
+          "compliance": {
+            "enum": [
+              "NotAvailable",
+              "NotCompliant",
+              "Compliant",
+              "Exempt"
+            ],
+            "type": "string"
+          },
+          "make": {
+            "type": "string"
+          },
+          "model": {
+            "type": "string"
+          },
+          "type": {
+            "type": "string"
+          },
+          "vrm": {
+            "type": "string"
+          }
+        },
+        "type": "object"
+      }
+    },
+    "securitySchemes": {
+      "apiKey": {
+        "description": "API Key Authentication",
+        "in": "query",
+        "name": "app_key",
+        "type": "apiKey"
+      },
+      "appId": {
+        "description": "Application ID",
+        "in": "query",
+        "name": "app_id",
+        "type": "apiKey"
+      }
+    }
+  }
+}
\ No newline at end of file
