WebAPI missing files pt2

diff --git a/Engine/Plugins/Web/WebAPI/Source/WebAPIOpenAPI/Private/Tests/Samples/V3/twitter_openapi.json b/Engine/Plugins/Web/WebAPI/Source/WebAPIOpenAPI/Private/Tests/Samples/V3/twitter_openapi.json
new file mode 100644
--- /dev/null
+++ b/Engine/Plugins/Web/WebAPI/Source/WebAPIOpenAPI/Private/Tests/Samples/V3/twitter_openapi.json
@@ -0,0 +1,5899 @@
+{
+  "openapi": "3.0.0",
+  "servers": [
+    {
+      "description": "Twitter API",
+      "url": "https://api.twitter.com"
+    }
+  ],
+  "info": {
+    "contact": {
+      "name": "Twitter Developers",
+      "url": "https://developer.twitter.com/",
+      "x-twitter": "twitter"
+    },
+    "description": "Twitter API v2 available endpoints",
+    "license": {
+      "name": "Twitter Developer Agreement and Policy",
+      "url": "https://developer.twitter.com/en/developer-terms/agreement-and-policy.html"
+    },
+    "termsOfService": "https://developer.twitter.com/en/developer-terms/agreement-and-policy.html",
+    "title": "Twitter API v2",
+    "version": "2.21",
+    "x-apisguru-categories": [
+      "social"
+    ],
+    "x-logo": {
+      "url": "https://api.apis.guru/v2/cache/logo/https_twitter.com_twitter_profile_image.png"
+    },
+    "x-origin": [
+      {
+        "format": "openapi",
+        "url": "https://api.twitter.com/2/openapi.json",
+        "version": "3.0"
+      }
+    ],
+    "x-providerName": "twitter.com",
+    "x-serviceName": "current",
+    "x-tags": [
+      "pii"
+    ]
+  },
+  "tags": [
+    {
+      "description": "Endpoints related to keeping Twitter data in your systems compliant",
+      "externalDocs": {
+        "description": "Find out more",
+        "url": "https://developer.twitter.com/en/docs/twitter-api/compliance/batch-tweet/introduction"
+      },
+      "name": "Compliance"
+    },
+    {
+      "description": "Miscellaneous endpoints for general API functionality",
+      "externalDocs": {
+        "description": "Find out more",
+        "url": "https://developer.twitter.com/en/docs/twitter-api"
+      },
+      "name": "General"
+    },
+    {
+      "description": "Endpoints related to retrieving, managing Lists",
+      "externalDocs": {
+        "description": "Find out more",
+        "url": "https://developer.twitter.com/en/docs/twitter-api/lists"
+      },
+      "name": "Lists"
+    },
+    {
+      "description": "Endpoints related to retrieving, managing Spaces",
+      "externalDocs": {
+        "description": "Find out more",
+        "url": "https://developer.twitter.com/en/docs/twitter-api/spaces"
+      },
+      "name": "Spaces"
+    },
+    {
+      "description": "Endpoints related to retrieving, searching, and modifying Tweets",
+      "externalDocs": {
+        "description": "Find out more",
+        "url": "https://developer.twitter.com/en/docs/twitter-api/tweets/lookup"
+      },
+      "name": "Tweets"
+    },
+    {
+      "description": "Endpoints related to retrieving, managing relationships of Users",
+      "externalDocs": {
+        "description": "Find out more",
+        "url": "https://developer.twitter.com/en/docs/twitter-api/users/lookup"
+      },
+      "name": "Users"
+    }
+  ],
+  "paths": {
+    "/2/openapi.json": {
+      "get": {
+        "description": "Full open api spec in JSON format. (See [REDACTED_PROFILE_URL]/OpenAPI-Specification/blob/master/README.md)",
+        "operationId": "getOpenApiSpec",
+        "parameters": [],
+        "responses": {
+          "200": {
+            "content": {
+              "application/json": {
+                "schema": {
+                  "type": "object"
+                }
+              }
+            },
+            "description": "The request was successful"
+          }
+        },
+        "summary": "Returns the open api spec document.",
+        "tags": [
+          "General"
+        ]
+      }
+    },
+    "/2/spaces": {
+      "get": {
+        "description": "Returns a variety of information about the Spaces specified by the requested IDs",
+        "operationId": "findSpacesByIds",
+        "parameters": [
+          {
+            "description": "A list of space ids",
+            "in": "query",
+            "name": "ids",
+            "required": true,
+            "schema": {
+              "items": {
+                "$ref": "#/components/schemas/SpaceID"
+              },
+              "minItems": 1,
+              "type": "array"
+            }
+          },
+          {
+            "$ref": "#/components/parameters/SpaceFieldsParameter"
+          },
+          {
+            "$ref": "#/components/parameters/SpaceExpansionsParameter"
+          }
+        ],
+        "responses": {
+          "200": {
+            "content": {
+              "application/json": {
+                "schema": {
+                  "$ref": "#/components/schemas/MultiSpaceLookupResponse"
+                }
+              }
+            },
+            "description": "The request was successful"
+          },
+          "default": {
+            "$ref": "#/components/responses/HttpErrorResponse"
+          }
+        },
+        "security": [
+          {
+            "BearerToken": [],
+            "UserToken": []
+          }
+        ],
+        "summary": "Space lookup up Space IDs",
+        "tags": [
+          "Spaces"
+        ]
+      }
+    },
+    "/2/spaces/by/creator_ids": {
+      "get": {
+        "description": "Returns a variety of information about the Spaces created by the provided User IDs",
+        "operationId": "findSpacesByCreatorIds",
+        "parameters": [
+          {
+            "description": "The users to search through",
+            "in": "query",
+            "name": "user_ids",
+            "required": true,
+            "schema": {
+              "items": {
+                "$ref": "#/components/schemas/UserID"
+              },
+              "minItems": 1,
+              "type": "array"
+            }
+          },
+          {
+            "$ref": "#/components/parameters/SpaceFieldsParameter"
+          },
+          {
+            "$ref": "#/components/parameters/SpaceExpansionsParameter"
+          }
+        ],
+        "responses": {
+          "200": {
+            "content": {
+              "application/json": {
+                "schema": {
+                  "$ref": "#/components/schemas/MultiSpaceLookupResponse"
+                }
+              }
+            },
+            "description": "The request was successful"
+          },
+          "default": {
+            "$ref": "#/components/responses/HttpErrorResponse"
+          }
+        },
+        "security": [
+          {
+            "BearerToken": [],
+            "UserToken": []
+          }
+        ],
+        "summary": "Space lookup by their creators",
+        "tags": [
+          "Spaces"
+        ]
+      }
+    },
+    "/2/spaces/search": {
+      "get": {
+        "description": "Returns Spaces that match the provided query.",
+        "operationId": "searchSpaces",
+        "parameters": [
+          {
+            "description": "The search query",
+            "example": "crypto",
+            "in": "query",
+            "name": "query",
+            "required": true,
+            "schema": {
+              "description": "The query to use in the Spaces search endpoint",
+              "example": "crypto",
+              "type": "string"
+            }
+          },
+          {
+            "description": "The state of spaces to search for",
+            "in": "query",
+            "name": "state",
+            "required": true,
+            "schema": {
+              "description": "The current state of the space.",
+              "enum": [
+                "live",
+                "scheduled"
+              ],
+              "example": "live",
+              "type": "string"
+            }
+          },
+          {
+            "description": "The number of results to return. The maximum for this value is 100.",
+            "in": "query",
+            "name": "max_results",
+            "required": false,
+            "schema": {
+              "description": "The number of results to return",
+              "example": 25,
+              "format": "int32",
+              "maximum": 100,
+              "minimum": 1,
+              "type": "integer"
+            }
+          },
+          {
+            "$ref": "#/components/parameters/SpaceFieldsParameter"
+          },
+          {
+            "$ref": "#/components/parameters/SpaceExpansionsParameter"
+          }
+        ],
+        "responses": {
+          "200": {
+            "content": {
+              "application/json": {
+                "schema": {
+                  "$ref": "#/components/schemas/MultiSpaceLookupResponse"
+                }
+              }
+            },
+            "description": "The request was successful"
+          },
+          "default": {
+            "$ref": "#/components/responses/HttpErrorResponse"
+          }
+        },
+        "security": [
+          {
+            "BearerToken": [],
+            "UserToken": []
+          }
+        ],
+        "summary": "Search for Spaces",
+        "tags": [
+          "Spaces"
+        ]
+      }
+    },
+    "/2/spaces/{id}": {
+      "get": {
+        "description": "Returns a variety of information about the Space specified by the requested ID",
+        "operationId": "findSpaceById",
+        "parameters": [
+          {
+            "description": "The space id to be retrieved",
+            "in": "path",
+            "name": "id",
+            "required": true,
+            "schema": {
+              "$ref": "#/components/schemas/SpaceID"
+            }
+          },
+          {
+            "$ref": "#/components/parameters/SpaceFieldsParameter"
+          },
+          {
+            "$ref": "#/components/parameters/SpaceExpansionsParameter"
+          }
+        ],
+        "responses": {
+          "200": {
+            "content": {
+              "application/json": {
+                "schema": {
+                  "$ref": "#/components/schemas/SingleSpaceLookupResponse"
+                }
+              }
+            },
+            "description": "The request was successful"
+          },
+          "default": {
+            "$ref": "#/components/responses/HttpErrorResponse"
+          }
+        },
+        "security": [
+          {
+            "BearerToken": [],
+            "UserToken": []
+          }
+        ],
+        "summary": "Space lookup by Space ID",
+        "tags": [
+          "Spaces"
+        ]
+      }
+    },
+    "/2/tweets": {
+      "get": {
+        "description": "Returns a variety of information about the Tweet specified by the requested ID",
+        "operationId": "findTweetsById",
+        "parameters": [
+          {
+            "description": "A comma separated list of Tweet IDs. Up to 100 are allowed in a single request.",
+            "explode": false,
+            "in": "query",
+            "name": "ids",
+            "required": true,
+            "schema": {
+              "items": {
+                "$ref": "#/components/schemas/TweetID"
+              },
+              "maxItems": 100,
+              "minItems": 1,
+              "type": "array"
+            },
+            "style": "form"
+          },
+          {
+            "$ref": "#/components/parameters/TweetExpansionsParameter"
+          },
+          {
+            "$ref": "#/components/parameters/TweetFieldsParameter"
+          },
+          {
+            "$ref": "#/components/parameters/UserFieldsParameter"
+          },
+          {
+            "$ref": "#/components/parameters/MediaFieldsParameter"
+          },
+          {
+            "$ref": "#/components/parameters/PlaceFieldsParameter"
+          },
+          {
+            "$ref": "#/components/parameters/PollFieldsParameter"
+          }
+        ],
+        "responses": {
+          "200": {
+            "content": {
+              "application/json": {
+                "schema": {
+                  "$ref": "#/components/schemas/MultiTweetLookupResponse"
+                }
+              }
+            },
+            "description": "The request was successful"
+          },
+          "default": {
+            "$ref": "#/components/responses/HttpErrorResponse"
+          }
+        },
+        "security": [
+          {
+            "BearerToken": [],
+            "UserToken": []
+          }
+        ],
+        "summary": "Tweet lookup by Tweet IDs",
+        "tags": [
+          "Tweets"
+        ]
+      }
+    },
+    "/2/tweets/counts/all": {
+      "get": {
+        "description": "Returns Tweet Counts that match a search query.",
+        "operationId": "tweetCountsFullArchiveSearch",
+        "parameters": [
+          {
+            "description": "One query/rule/filter for matching Tweets. Up to 2048 characters.",
+            "in": "query",
+            "name": "query",
+            "required": true,
+            "schema": {
+              "example": "(from:TwitterDev OR from:TwitterAPI) has:media -is:retweet",
+              "maxLength": 2048,
+              "minLength": 1,
+              "type": "string"
+            }
+          },
+          {
+            "description": "YYYY-MM-DDTHH:mm:ssZ. The oldest UTC timestamp (from most recent 7 days) from which the Tweets will be provided. Timestamp is in second granularity and is inclusive (i.e. 12:00:01 includes the first second of the minute).",
+            "in": "query",
+            "name": "start_time",
+            "required": false,
+            "schema": {
+              "format": "date-time",
+              "type": "string"
+            }
+          },
+          {
+            "description": "YYYY-MM-DDTHH:mm:ssZ. The newest, most recent UTC timestamp to which the Tweets will be provided. Timestamp is in second granularity and is exclusive (i.e. 12:00:01 excludes the first second of the minute).",
+            "in": "query",
+            "name": "end_time",
+            "required": false,
+            "schema": {
+              "format": "date-time",
+              "type": "string"
+            }
+          },
+          {
+            "description": "Returns results with a Tweet ID greater than (that is, more recent than) the specified ID.",
+            "in": "query",
+            "name": "since_id",
+            "required": false,
+            "schema": {
+              "$ref": "#/components/schemas/TweetID"
+            }
+          },
+          {
+            "description": "Returns results with a Tweet ID less than (that is, older than) the specified ID.",
+            "in": "query",
+            "name": "until_id",
+            "required": false,
+            "schema": {
+              "$ref": "#/components/schemas/TweetID"
+            }
+          },
+          {
+            "description": "This parameter is used to get the next 'page' of results. The value used with the parameter is pulled directly from the response provided by the API, and should not be modified.",
+            "in": "query",
+            "name": "next_token",
+            "required": false,
+            "schema": {
+              "minLength": 1,
+              "type": "string"
+            }
+          },
+          {
+            "description": "The granularity for the search counts results.",
+            "in": "query",
+            "name": "granularity",
+            "required": false,
+            "schema": {
+              "$ref": "#/components/schemas/Granularity"
+            }
+          }
+        ],
+        "responses": {
+          "200": {
+            "content": {
+              "application/json": {
+                "schema": {
+                  "$ref": "#/components/schemas/TweetCountsResponse"
+                }
+              }
+            },
+            "description": "Tweet counts response"
+          },
+          "default": {
+            "$ref": "#/components/responses/HttpErrorResponse"
+          }
+        },
+        "security": [
+          {
+            "BearerToken": [],
+            "UserToken": []
+          }
+        ],
+        "summary": "Full archive search counts",
+        "tags": [
+          "Tweets"
+        ]
+      }
+    },
+    "/2/tweets/counts/recent": {
+      "get": {
+        "description": "Returns Tweet Counts from the last 7 days that match a search query.",
+        "operationId": "tweetCountsRecentSearch",
+        "parameters": [
+          {
+            "description": "One query/rule/filter for matching Tweets. Up to 2048 characters.",
+            "in": "query",
+            "name": "query",
+            "required": true,
+            "schema": {
+              "example": "(from:TwitterDev OR from:TwitterAPI) has:media -is:retweet",
+              "maxLength": 2048,
+              "minLength": 1,
+              "type": "string"
+            }
+          },
+          {
+            "description": "YYYY-MM-DDTHH:mm:ssZ. The oldest UTC timestamp (from most recent 7 days) from which the Tweets will be provided. Timestamp is in second granularity and is inclusive (i.e. 12:00:01 includes the first second of the minute).",
+            "in": "query",
+            "name": "start_time",
+            "required": false,
+            "schema": {
+              "format": "date-time",
+              "type": "string"
+            }
+          },
+          {
+            "description": "YYYY-MM-DDTHH:mm:ssZ. The newest, most recent UTC timestamp to which the Tweets will be provided. Timestamp is in second granularity and is exclusive (i.e. 12:00:01 excludes the first second of the minute).",
+            "in": "query",
+            "name": "end_time",
+            "required": false,
+            "schema": {
+              "format": "date-time",
+              "type": "string"
+            }
+          },
+          {
+            "description": "Returns results with a Tweet ID greater than (that is, more recent than) the specified ID.",
+            "in": "query",
+            "name": "since_id",
+            "required": false,
+            "schema": {
+              "$ref": "#/components/schemas/TweetID"
+            }
+          },
+          {
+            "description": "Returns results with a Tweet ID less than (that is, older than) the specified ID.",
+            "in": "query",
+            "name": "until_id",
+            "required": false,
+            "schema": {
+              "$ref": "#/components/schemas/TweetID"
+            }
+          },
+          {
+            "description": "This parameter is used to get the next 'page' of results. The value used with the parameter is pulled directly from the response provided by the API, and should not be modified.",
+            "in": "query",
+            "name": "next_token",
+            "required": false,
+            "schema": {
+              "minLength": 1,
+              "type": "string"
+            }
+          },
+          {
+            "description": "The granularity for the search counts results.",
+            "in": "query",
+            "name": "granularity",
+            "required": false,
+            "schema": {
+              "$ref": "#/components/schemas/Granularity"
+            }
+          }
+        ],
+        "responses": {
+          "200": {
+            "content": {
+              "application/json": {
+                "schema": {
+                  "$ref": "#/components/schemas/TweetCountsResponse"
+                }
+              }
+            },
+            "description": "Recent tweet counts response"
+          },
+          "default": {
+            "$ref": "#/components/responses/HttpErrorResponse"
+          }
+        },
+        "security": [
+          {
+            "BearerToken": [],
+            "UserToken": []
+          }
+        ],
+        "summary": "Recent search counts",
+        "tags": [
+          "Tweets"
+        ]
+      }
+    },
+    "/2/tweets/sample/stream": {
+      "get": {
+        "description": "Streams a deterministic 1% of public Tweets.",
+        "operationId": "sampleStream",
+        "parameters": [
+          {
+            "$ref": "#/components/parameters/TweetExpansionsParameter"
+          },
+          {
+            "$ref": "#/components/parameters/TweetFieldsParameter"
+          },
+          {
+            "$ref": "#/components/parameters/UserFieldsParameter"
+          },
+          {
+            "$ref": "#/components/parameters/MediaFieldsParameter"
+          },
+          {
+            "$ref": "#/components/parameters/PlaceFieldsParameter"
+          },
+          {
+            "$ref": "#/components/parameters/PollFieldsParameter"
+          },
+          {
+            "$ref": "#/components/parameters/BackfillMinutesRequestParameter"
+          }
+        ],
+        "responses": {
+          "200": {
+            "content": {
+              "application/json": {
+                "schema": {
+                  "$ref": "#/components/schemas/StreamingTweet"
+                }
+              }
+            },
+            "description": "The request was successful. Successful responses will return a stream of individual JSON Tweet payloads."
+          },
+          "default": {
+            "$ref": "#/components/responses/HttpErrorResponse"
+          }
+        },
+        "security": [
+          {
+            "BearerToken": []
+          }
+        ],
+        "summary": "Sample stream",
+        "tags": [
+          "Tweets"
+        ],
+        "x-twitter-streaming": true
+      }
+    },
+    "/2/tweets/search/all": {
+      "get": {
+        "description": "Returns Tweets that match a search query.",
+        "operationId": "tweetsFullarchiveSearch",
+        "parameters": [
+          {
+            "description": "One query/rule/filter for matching Tweets. Up to 1024 characters.",
+            "in": "query",
+            "name": "query",
+            "required": true,
+            "schema": {
+              "example": "(from:TwitterDev OR from:TwitterAPI) has:media -is:retweet",
+              "maxLength": 1024,
+              "minLength": 1,
+              "type": "string"
+            }
+          },
+          {
+            "description": "YYYY-MM-DDTHH:mm:ssZ. The oldest UTC timestamp from which the Tweets will be provided. Timestamp is in second granularity and is inclusive (i.e. 12:00:01 includes the first second of the minute).",
+            "in": "query",
+            "name": "start_time",
+            "required": false,
+            "schema": {
+              "format": "date-time",
+              "type": "string"
+            }
+          },
+          {
+            "description": "YYYY-MM-DDTHH:mm:ssZ. The newest, most recent UTC timestamp to which the Tweets will be provided. Timestamp is in second granularity and is exclusive (i.e. 12:00:01 excludes the first second of the minute).",
+            "in": "query",
+            "name": "end_time",
+            "required": false,
+            "schema": {
+              "format": "date-time",
+              "type": "string"
+            }
+          },
+          {
+            "description": "Returns results with a Tweet ID greater than (that is, more recent than) the specified ID.",
+            "in": "query",
+            "name": "since_id",
+            "required": false,
+            "schema": {
+              "$ref": "#/components/schemas/TweetID"
+            }
+          },
+          {
+            "description": "Returns results with a Tweet ID less than (that is, older than) the specified ID.",
+            "in": "query",
+            "name": "until_id",
+            "required": false,
+            "schema": {
+              "$ref": "#/components/schemas/TweetID"
+            }
+          },
+          {
+            "description": "The maximum number of search results to be returned by a request.",
+            "in": "query",
+            "name": "max_results",
+            "required": false,
+            "schema": {
+              "default": 10,
+              "format": "int32",
+              "maximum": 500,
+              "minimum": 10,
+              "type": "integer"
+            }
+          },
+          {
+            "description": "This parameter is used to get the next 'page' of results. The value used with the parameter is pulled directly from the response provided by the API, and should not be modified.",
+            "in": "query",
+            "name": "next_token",
+            "required": false,
+            "schema": {
+              "minLength": 1,
+              "type": "string"
+            }
+          },
+          {
+            "$ref": "#/components/parameters/TweetExpansionsParameter"
+          },
+          {
+            "$ref": "#/components/parameters/TweetFieldsParameter"
+          },
+          {
+            "$ref": "#/components/parameters/UserFieldsParameter"
+          },
+          {
+            "$ref": "#/components/parameters/MediaFieldsParameter"
+          },
+          {
+            "$ref": "#/components/parameters/PlaceFieldsParameter"
+          },
+          {
+            "$ref": "#/components/parameters/PollFieldsParameter"
+          }
+        ],
+        "responses": {
+          "200": {
+            "content": {
+              "application/json": {
+                "schema": {
+                  "$ref": "#/components/schemas/TweetSearchResponse"
+                }
+              }
+            },
+            "description": "Tweets full archive search response"
+          },
+          "default": {
+            "$ref": "#/components/responses/HttpErrorResponse"
+          }
+        },
+        "security": [
+          {
+            "BearerToken": []
+          }
+        ],
+        "summary": "Full-archive search",
+        "tags": [
+          "Tweets"
+        ]
+      }
+    },
+    "/2/tweets/search/recent": {
+      "get": {
+        "description": "Returns Tweets from the last 7 days that match a search query.",
+        "operationId": "tweetsRecentSearch",
+        "parameters": [
+          {
+            "description": "One query/rule/filter for matching Tweets. Up to 512 characters.",
+            "in": "query",
+            "name": "query",
+            "required": true,
+            "schema": {
+              "example": "(from:TwitterDev OR from:TwitterAPI) has:media -is:retweet",
+              "maxLength": 512,
+              "minLength": 1,
+              "type": "string"
+            }
+          },
+          {
+            "description": "YYYY-MM-DDTHH:mm:ssZ. The oldest UTC timestamp (from most recent 7 days) from which the Tweets will be provided. Timestamp is in second granularity and is inclusive (i.e. 12:00:01 includes the first second of the minute).",
+            "in": "query",
+            "name": "start_time",
+            "required": false,
+            "schema": {
+              "format": "date-time",
+              "type": "string"
+            }
+          },
+          {
+            "description": "YYYY-MM-DDTHH:mm:ssZ. The newest, most recent UTC timestamp to which the Tweets will be provided. Timestamp is in second granularity and is exclusive (i.e. 12:00:01 excludes the first second of the minute).",
+            "in": "query",
+            "name": "end_time",
+            "required": false,
+            "schema": {
+              "format": "date-time",
+              "type": "string"
+            }
+          },
+          {
+            "description": "Returns results with a Tweet ID greater than (that is, more recent than) the specified ID.",
+            "in": "query",
+            "name": "since_id",
+            "required": false,
+            "schema": {
+              "$ref": "#/components/schemas/TweetID"
+            }
+          },
+          {
+            "description": "Returns results with a Tweet ID less than (that is, older than) the specified ID.",
+            "in": "query",
+            "name": "until_id",
+            "required": false,
+            "schema": {
+              "$ref": "#/components/schemas/TweetID"
+            }
+          },
+          {
+            "description": "The maximum number of search results to be returned by a request.",
+            "in": "query",
+            "name": "max_results",
+            "required": false,
+            "schema": {
+              "default": 10,
+              "format": "int32",
+              "maximum": 100,
+              "minimum": 10,
+              "type": "integer"
+            }
+          },
+          {
+            "description": "This parameter is used to get the next 'page' of results. The value used with the parameter is pulled directly from the response provided by the API, and should not be modified.",
+            "in": "query",
+            "name": "next_token",
+            "required": false,
+            "schema": {
+              "minLength": 1,
+              "type": "string"
+            }
+          },
+          {
+            "$ref": "#/components/parameters/TweetExpansionsParameter"
+          },
+          {
+            "$ref": "#/components/parameters/TweetFieldsParameter"
+          },
+          {
+            "$ref": "#/components/parameters/UserFieldsParameter"
+          },
+          {
+            "$ref": "#/components/parameters/MediaFieldsParameter"
+          },
+          {
+            "$ref": "#/components/parameters/PlaceFieldsParameter"
+          },
+          {
+            "$ref": "#/components/parameters/PollFieldsParameter"
+          }
+        ],
+        "responses": {
+          "200": {
+            "content": {
+              "application/json": {
+                "schema": {
+                  "$ref": "#/components/schemas/TweetSearchResponse"
+                }
+              }
+            },
+            "description": "Tweets recent search response"
+          },
+          "default": {
+            "$ref": "#/components/responses/HttpErrorResponse"
+          }
+        },
+        "security": [
+          {
+            "BearerToken": [],
+            "UserToken": []
+          }
+        ],
+        "summary": "Recent search",
+        "tags": [
+          "Tweets"
+        ]
+      }
+    },
+    "/2/tweets/search/stream": {
+      "get": {
+        "description": "Streams Tweets matching the stream's active rule set.",
+        "operationId": "searchStream",
+        "parameters": [
+          {
+            "$ref": "#/components/parameters/TweetExpansionsParameter"
+          },
+          {
+            "$ref": "#/components/parameters/TweetFieldsParameter"
+          },
+          {
+            "$ref": "#/components/parameters/UserFieldsParameter"
+          },
+          {
+            "$ref": "#/components/parameters/MediaFieldsParameter"
+          },
+          {
+            "$ref": "#/components/parameters/PlaceFieldsParameter"
+          },
+          {
+            "$ref": "#/components/parameters/PollFieldsParameter"
+          },
+          {
+            "$ref": "#/components/parameters/BackfillMinutesRequestParameter"
+          }
+        ],
+        "responses": {
+          "200": {
+            "content": {
+              "application/json": {
+                "schema": {
+                  "$ref": "#/components/schemas/FilteredStreamingTweet"
+                }
+              }
+            },
+            "description": "The request was successful. Successful responses will return a stream of individual JSON Tweet payloads."
+          },
+          "default": {
+            "$ref": "#/components/responses/HttpErrorResponse"
+          }
+        },
+        "security": [
+          {
+            "BearerToken": []
+          }
+        ],
+        "summary": "Filtered stream",
+        "tags": [
+          "Tweets"
+        ],
+        "x-twitter-streaming": true
+      }
+    },
+    "/2/tweets/search/stream/rules": {
+      "get": {
+        "description": "Returns rules from a user's active rule set. Users can fetch all of their rules or a subset, specified by the provided rule ids.",
+        "operationId": "getRules",
+        "parameters": [
+          {
+            "description": "A comma-separated list of Rule IDs.",
+            "in": "query",
+            "name": "ids",
+            "required": false,
+            "schema": {
+              "items": {
+                "$ref": "#/components/schemas/RuleId"
+              },
+              "type": "array"
+            }
+          }
+        ],
+        "responses": {
+          "200": {
+            "content": {
+              "application/json": {
+                "schema": {
+                  "properties": {
+                    "data": {
+                      "items": {
+                        "$ref": "#/components/schemas/Rule"
+                      },
+                      "type": "array"
+                    },
+                    "meta": {
+                      "$ref": "#/components/schemas/RulesResponseMetadata"
+                    }
+                  },
+                  "required": [
+                    "data",
+                    "meta"
+                  ],
+                  "type": "object"
+                }
+              }
+            },
+            "description": "The request was successful"
+          },
+          "default": {
+            "$ref": "#/components/responses/HttpErrorResponse"
+          }
+        },
+        "security": [
+          {
+            "BearerToken": []
+          }
+        ],
+        "summary": "Rules lookup",
+        "tags": [
+          "Tweets"
+        ]
+      },
+      "post": {
+        "description": "Add or delete rules from a user's active rule set. Users can provide unique, optionally tagged rules to add. Users can delete their entire rule set or a subset specified by rule ids or values.",
+        "operationId": "addOrDeleteRules",
+        "parameters": [
+          {
+            "description": "Dry Run can be used with both the add and delete action, with the expected result given, but without actually taking any action in the system (meaning the end state will always be as it was when the request was submitted). This is particularly useful to validate rule changes.",
+            "in": "query",
+            "name": "dry_run",
+            "required": false,
+            "schema": {
+              "type": "boolean"
+            }
+          }
+        ],
+        "requestBody": {
+          "content": {
+            "application/json": {
+              "schema": {
+                "$ref": "#/components/schemas/AddOrDeleteRulesRequest"
+              }
+            }
+          },
+          "description": "",
+          "required": true
+        },
+        "responses": {
+          "200": {
+            "content": {
+              "application/json": {
+                "schema": {
+                  "$ref": "#/components/schemas/AddOrDeleteRulesResponse"
+                }
+              }
+            },
+            "description": "The request was successful"
+          },
+          "default": {
+            "$ref": "#/components/responses/HttpErrorResponse"
+          }
+        },
+        "security": [
+          {
+            "BearerToken": []
+          }
+        ],
+        "summary": "Add/Delete rules",
+        "tags": [
+          "Tweets"
+        ]
+      }
+    },
+    "/2/tweets/{id}": {
+      "get": {
+        "description": "Returns a variety of information about the Tweet specified by the requested ID",
+        "operationId": "findTweetById",
+        "parameters": [
+          {
+            "description": "A single Tweet ID.",
+            "in": "path",
+            "name": "id",
+            "required": true,
+            "schema": {
+              "$ref": "#/components/schemas/TweetID"
+            }
+          },
+          {
+            "$ref": "#/components/parameters/TweetExpansionsParameter"
+          },
+          {
+            "$ref": "#/components/parameters/TweetFieldsParameter"
+          },
+          {
+            "$ref": "#/components/parameters/UserFieldsParameter"
+          },
+          {
+            "$ref": "#/components/parameters/MediaFieldsParameter"
+          },
+          {
+            "$ref": "#/components/parameters/PlaceFieldsParameter"
+          },
+          {
+            "$ref": "#/components/parameters/PollFieldsParameter"
+          }
+        ],
+        "responses": {
+          "200": {
+            "content": {
+              "application/json": {
+                "schema": {
+                  "$ref": "#/components/schemas/SingleTweetLookupResponse"
+                }
+              }
+            },
+            "description": "The request was successful"
+          },
+          "default": {
+            "$ref": "#/components/responses/HttpErrorResponse"
+          }
+        },
+        "security": [
+          {
+            "BearerToken": [],
+            "UserToken": []
+          }
+        ],
+        "summary": "Tweet lookup by Tweet ID",
+        "tags": [
+          "Tweets"
+        ]
+      }
+    },
+    "/2/tweets/{id}/hidden": {
+      "put": {
+        "description": "Hides or unhides a reply to an owned conversation.",
+        "operationId": "hideReplyById",
+        "parameters": [
+          {
+            "description": "The ID of the reply that you want to hide or unhide.",
+            "in": "path",
+            "name": "id",
+            "required": true,
+            "schema": {
+              "$ref": "#/components/schemas/TweetID"
+            }
+          }
+        ],
+        "requestBody": {
+          "content": {
+            "application/json": {
+              "schema": {
+                "properties": {
+                  "hidden": {
+                    "type": "boolean"
+                  }
+                },
+                "type": "object"
+              }
+            }
+          }
+        },
+        "responses": {
+          "200": {
+            "content": {
+              "application/json": {
+                "schema": {
+                  "properties": {
+                    "data": {
+                      "properties": {
+                        "hidden": {
+                          "type": "boolean"
+                        }
+                      },
+                      "type": "object"
+                    }
+                  },
+                  "type": "object"
+                }
+              }
+            },
+            "description": "A successful response. The reply has been hidden or unhidden."
+          },
+          "default": {
+            "$ref": "#/components/responses/HttpErrorResponse"
+          }
+        },
+        "security": [
+          {
+            "UserToken": []
+          }
+        ],
+        "summary": "Hide replies",
+        "tags": [
+          "Tweets"
+        ]
+      }
+    },
+    "/2/tweets/{id}/liking_users": {
+      "get": {
+        "description": "Returns a list of users that have liked the provided Tweet ID",
+        "operationId": "tweetsIdLikingUsers",
+        "parameters": [
+          {
+            "description": "The ID of the Tweet for which to return results",
+            "in": "path",
+            "name": "id",
+            "required": true,
+            "schema": {
+              "$ref": "#/components/schemas/TweetID"
+            }
+          }
+        ],
+        "responses": {
+          "200": {
+            "content": {
+              "application/json": {
+                "schema": {
+                  "$ref": "#/components/schemas/GenericMultipleUsersLookupResponse"
+                }
+              }
+            },
+            "description": "The request was successful"
+          },
+          "default": {
+            "$ref": "#/components/responses/HttpErrorResponse"
+          }
+        },
+        "security": [
+          {
+            "BearerToken": [],
+            "UserToken": []
+          }
+        ],
+        "summary": "Returns user objects that have liked the provided Tweet ID",
+        "tags": [
+          "Users"
+        ]
+      }
+    },
+    "/2/tweets/{id}/retweeted_by": {
+      "get": {
+        "description": "Returns a list of users that have retweeted the provided Tweet ID",
+        "operationId": "tweetsIdRetweetingUsers",
+        "parameters": [
+          {
+            "description": "The ID of the Tweet for which to return results",
+            "in": "path",
+            "name": "id",
+            "required": true,
+            "schema": {
+              "$ref": "#/components/schemas/TweetID"
+            }
+          }
+        ],
+        "responses": {
+          "200": {
+            "content": {
+              "application/json": {
+                "schema": {
+                  "$ref": "#/components/schemas/GenericMultipleUsersLookupResponse"
+                }
+              }
+            },
+            "description": "The request was successful"
+          },
+          "default": {
+            "$ref": "#/components/responses/HttpErrorResponse"
+          }
+        },
+        "security": [
+          {
+            "BearerToken": [],
+            "UserToken": []
+          }
+        ],
+        "summary": "Returns user objects that have retweeted the provided Tweet ID",
+        "tags": [
+          "Users"
+        ]
+      }
+    },
+    "/2/users": {
+      "get": {
+        "description": "This endpoint returns information about users. Specify users by their ID.",
+        "operationId": "findUsersById",
+        "parameters": [
+          {
+            "description": "Required. A list of User IDs, comma-separated. You can specify up to 100 IDs.",
+            "example": "2244994945,6253282,12",
+            "explode": false,
+            "in": "query",
+            "name": "ids",
+            "required": true,
+            "schema": {
+              "items": {
+                "$ref": "#/components/schemas/UserID"
+              },
+              "maxItems": 100,
+              "minItems": 1,
+              "type": "array"
+            },
+            "style": "form"
+          },
+          {
+            "$ref": "#/components/parameters/UserExpansionsParameter"
+          },
+          {
+            "$ref": "#/components/parameters/TweetFieldsParameter"
+          },
+          {
+            "$ref": "#/components/parameters/UserFieldsParameter"
+          }
+        ],
+        "responses": {
+          "200": {
+            "content": {
+              "application/json": {
+                "schema": {
+                  "$ref": "#/components/schemas/MultiUserLookupResponse"
+                }
+              }
+            },
+            "description": "The request was successful"
+          },
+          "default": {
+            "$ref": "#/components/responses/HttpErrorResponse"
+          }
+        },
+        "security": [
+          {
+            "BearerToken": [],
+            "UserToken": []
+          }
+        ],
+        "summary": "User lookup by IDs",
+        "tags": [
+          "Users"
+        ]
+      }
+    },
+    "/2/users/by": {
+      "get": {
+        "description": "This endpoint returns information about users. Specify users by their username.",
+        "operationId": "findUsersByUsername",
+        "parameters": [
+          {
+            "description": "Required . A list of usernames, comma-separated. You can specify up to 100 usernames.",
+            "example": "TwitterDev,TwitterAPI,jack",
+            "explode": false,
+            "in": "query",
+            "name": "usernames",
+            "required": true,
+            "schema": {
+              "items": {
+                "$ref": "#/components/schemas/UserName"
+              },
+              "maxItems": 100,
+              "minItems": 1,
+              "type": "array"
+            },
+            "style": "form"
+          },
+          {
+            "$ref": "#/components/parameters/UserExpansionsParameter"
+          },
+          {
+            "$ref": "#/components/parameters/TweetFieldsParameter"
+          },
+          {
+            "$ref": "#/components/parameters/UserFieldsParameter"
+          }
+        ],
+        "responses": {
+          "200": {
+            "content": {
+              "application/json": {
+                "schema": {
+                  "$ref": "#/components/schemas/MultiUserLookupResponse"
+                }
+              }
+            },
+            "description": "The request was successful"
+          },
+          "default": {
+            "$ref": "#/components/responses/HttpErrorResponse"
+          }
+        },
+        "security": [
+          {
+            "BearerToken": [],
+            "UserToken": []
+          }
+        ],
+        "summary": "User lookup by usernames",
+        "tags": [
+          "Users"
+        ]
+      }
+    },
+    "/2/users/by/username/{username}": {
+      "get": {
+        "description": "This endpoint returns information about a user. Specify user by username.",
+        "operationId": "findUserByUsername",
+        "parameters": [
+          {
+            "description": "Required. A username.",
+            "example": "TwitterDev",
+            "in": "path",
+            "name": "username",
+            "required": true,
+            "schema": {
+              "$ref": "#/components/schemas/UserName"
+            }
+          },
+          {
+            "$ref": "#/components/parameters/UserExpansionsParameter"
+          },
+          {
+            "$ref": "#/components/parameters/TweetFieldsParameter"
+          },
+          {
+            "$ref": "#/components/parameters/UserFieldsParameter"
+          }
+        ],
+        "responses": {
+          "200": {
+            "content": {
+              "application/json": {
+                "schema": {
+                  "$ref": "#/components/schemas/SingleUserLookupResponse"
+                }
+              }
+            },
+            "description": "The request was successful"
+          },
+          "default": {
+            "$ref": "#/components/responses/HttpErrorResponse"
+          }
+        },
+        "security": [
+          {
+            "BearerToken": [],
+            "UserToken": []
+          }
+        ],
+        "summary": "User lookup by username",
+        "tags": [
+          "Users"
+        ]
+      }
+    },
+    "/2/users/{id}": {
+      "get": {
+        "description": "This endpoint returns information about a user. Specify user by ID.",
+        "operationId": "findUserById",
+        "parameters": [
+          {
+            "description": "Required. A User ID.",
+            "example": "2244994945",
+            "in": "path",
+            "name": "id",
+            "required": true,
+            "schema": {
+              "$ref": "#/components/schemas/UserID"
+            }
+          },
+          {
+            "$ref": "#/components/parameters/UserExpansionsParameter"
+          },
+          {
+            "$ref": "#/components/parameters/TweetFieldsParameter"
+          },
+          {
+            "$ref": "#/components/parameters/UserFieldsParameter"
+          }
+        ],
+        "responses": {
+          "200": {
+            "content": {
+              "application/json": {
+                "schema": {
+                  "$ref": "#/components/schemas/SingleUserLookupResponse"
+                }
+              }
+            },
+            "description": "The request was successful"
+          },
+          "default": {
+            "$ref": "#/components/responses/HttpErrorResponse"
+          }
+        },
+        "security": [
+          {
+            "BearerToken": [],
+            "UserToken": []
+          }
+        ],
+        "summary": "User lookup by ID",
+        "tags": [
+          "Users"
+        ]
+      }
+    },
+    "/2/users/{id}/blocking": {
+      "get": {
+        "description": "Returns a list of users that are blocked by the provided user ID",
+        "operationId": "usersIdBlocking",
+        "parameters": [
+          {
+            "description": "The ID of the user for whom to return results",
+            "in": "path",
+            "name": "id",
+            "required": true,
+            "schema": {
+              "$ref": "#/components/schemas/UserID"
+            }
+          },
+          {
+            "description": "The maximum number of results",
+            "in": "query",
+            "name": "max_results",
+            "required": false,
+            "schema": {
+              "format": "int32",
+              "maximum": 1000,
+              "minimum": 1,
+              "type": "integer"
+            }
+          },
+          {
+            "description": "This value is populated by passing the 'next_token' or 'previous_token' returned in a request to paginate through results.",
+            "in": "query",
+            "name": "pagination_token",
+            "required": false,
+            "schema": {
+              "maxLength": 16,
+              "minLength": 16,
+              "type": "string"
+            }
+          }
+        ],
+        "responses": {
+          "200": {
+            "content": {
+              "application/json": {
+                "schema": {
+                  "$ref": "#/components/schemas/GenericMultipleUsersLookupResponse"
+                }
+              }
+            },
+            "description": "The request was successful"
+          },
+          "default": {
+            "$ref": "#/components/responses/HttpErrorResponse"
+          }
+        },
+        "security": [
+          {
+            "UserToken": []
+          }
+        ],
+        "summary": "Returns user objects that are blocked by provided user ID",
+        "tags": [
+          "Users"
+        ]
+      },
+      "post": {
+        "description": "Causes the user (in the path) to block the target user. The user (in the path) must match the user context authorizing the request",
+        "operationId": "usersIdBlock",
+        "parameters": [
+          {
+            "description": "The ID of the user that is requesting to block the target user",
+            "in": "path",
+            "name": "id",
+            "required": true,
+            "schema": {
+              "$ref": "#/components/schemas/UserID"
+            }
+          }
+        ],
+        "requestBody": {
+          "content": {
+            "application/json": {
+              "schema": {
+                "properties": {
+                  "target_user_id": {
+                    "pattern": "^[0-9]{1,19}$",
+                    "type": "string"
+                  }
+                },
+                "required": [
+                  "target_user_id"
+                ],
+                "type": "object"
+              }
+            }
+          }
+        },
+        "responses": {
+          "200": {
+            "content": {
+              "application/json": {
+                "schema": {
+                  "$ref": "#/components/schemas/UsersBlockingMutationResponse"
+                }
+              }
+            },
+            "description": "The request was successful"
+          },
+          "default": {
+            "$ref": "#/components/responses/HttpErrorResponse"
+          }
+        },
+        "security": [
+          {
+            "UserToken": []
+          }
+        ],
+        "summary": "Block User by User ID",
+        "tags": [
+          "Users"
+        ]
+      }
+    },
+    "/2/users/{id}/followers": {
+      "get": {
+        "description": "Returns a list of users that follow the provided user ID",
+        "operationId": "usersIdFollowers",
+        "parameters": [
+          {
+            "description": "The ID of the user for whom to return results",
+            "in": "path",
+            "name": "id",
+            "required": true,
+            "schema": {
+              "$ref": "#/components/schemas/UserID"
+            }
+          },
+          {
+            "description": "The maximum number of results",
+            "in": "query",
+            "name": "max_results",
+            "required": false,
+            "schema": {
+              "format": "int32",
+              "maximum": 1000,
+              "minimum": 1,
+              "type": "integer"
+            }
+          },
+          {
+            "description": "This value is populated by passing the 'next_token' or 'previous_token' returned in a request to paginate through results.",
+            "in": "query",
+            "name": "pagination_token",
+            "required": false,
+            "schema": {
+              "maxLength": 16,
+              "minLength": 16,
+              "type": "string"
+            }
+          }
+        ],
+        "responses": {
+          "200": {
+            "content": {
+              "application/json": {
+                "schema": {
+                  "$ref": "#/components/schemas/GenericMultipleUsersLookupResponse"
+                }
+              }
+            },
+            "description": "The request was successful"
+          },
+          "default": {
+            "$ref": "#/components/responses/HttpErrorResponse"
+          }
+        },
+        "security": [
+          {
+            "BearerToken": [],
+            "UserToken": []
+          }
+        ],
+        "summary": "Returns user objects that follow the provided user ID",
+        "tags": [
+          "Users"
+        ]
+      }
+    },
+    "/2/users/{id}/following": {
+      "get": {
+        "description": "Returns a list of users that are being followed by the provided user ID",
+        "operationId": "usersIdFollowing",
+        "parameters": [
+          {
+            "description": "The ID of the user for whom to return results",
+            "in": "path",
+            "name": "id",
+            "required": true,
+            "schema": {
+              "$ref": "#/components/schemas/UserID"
+            }
+          },
+          {
+            "description": "The maximum number of results",
+            "in": "query",
+            "name": "max_results",
+            "required": false,
+            "schema": {
+              "format": "int32",
+              "maximum": 1000,
+              "minimum": 1,
+              "type": "integer"
+            }
+          },
+          {
+            "description": "This value is populated by passing the 'next_token' or 'previous_token' returned in a request to paginate through results.",
+            "in": "query",
+            "name": "pagination_token",
+            "required": false,
+            "schema": {
+              "maxLength": 16,
+              "minLength": 16,
+              "type": "string"
+            }
+          }
+        ],
+        "responses": {
+          "200": {
+            "content": {
+              "application/json": {
+                "schema": {
+                  "$ref": "#/components/schemas/UsersFollowingLookupResponse"
+                }
+              }
+            },
+            "description": "The request was successful"
+          },
+          "default": {
+            "$ref": "#/components/responses/HttpErrorResponse"
+          }
+        },
+        "security": [
+          {
+            "BearerToken": [],
+            "UserToken": []
+          }
+        ],
+        "summary": "Following by User ID",
+        "tags": [
+          "Users"
+        ]
+      },
+      "post": {
+        "description": "Causes the user(in the path) to follow, or “request to follow” for protected users, the target user. The user(in the path) must match the user context authorizing the request",
+        "operationId": "usersIdFollow",
+        "parameters": [
+          {
+            "description": "The ID of the user that is requesting to follow the target user",
+            "in": "path",
+            "name": "id",
+            "required": true,
+            "schema": {
+              "$ref": "#/components/schemas/UserID"
+            }
+          }
+        ],
+        "requestBody": {
+          "content": {
+            "application/json": {
+              "schema": {
+                "properties": {
+                  "target_user_id": {
+                    "pattern": "^[0-9]{1,19}$",
+                    "type": "string"
+                  }
+                },
+                "required": [
+                  "target_user_id"
+                ],
+                "type": "object"
+              }
+            }
+          }
+        },
+        "responses": {
+          "200": {
+            "content": {
+              "application/json": {
+                "schema": {
+                  "$ref": "#/components/schemas/UsersFollowingCreateResponse"
+                }
+              }
+            },
+            "description": "The request was successful"
+          },
+          "default": {
+            "$ref": "#/components/responses/HttpErrorResponse"
+          }
+        },
+        "security": [
+          {
+            "UserToken": []
+          }
+        ],
+        "summary": "Follow User",
+        "tags": [
+          "Users"
+        ]
+      }
+    },
+    "/2/users/{id}/liked_tweets": {
+      "get": {
+        "description": "Returns a list of Tweets liked by the provided User ID",
+        "operationId": "usersIdLikedTweets",
+        "parameters": [
+          {
+            "description": "The ID of the User to list the liked Tweets of",
+            "in": "path",
+            "name": "id",
+            "required": true,
+            "schema": {
+              "$ref": "#/components/schemas/UserID"
+            }
+          },
+          {
+            "$ref": "#/components/parameters/MaxResultsRequestParameter"
+          },
+          {
+            "$ref": "#/components/parameters/PaginationTokenRequestParameter"
+          },
+          {
+            "$ref": "#/components/parameters/TweetExpansionsParameter"
+          },
+          {
+            "$ref": "#/components/parameters/TweetFieldsParameter"
+          },
+          {
+            "$ref": "#/components/parameters/UserFieldsParameter"
+          },
+          {
+            "$ref": "#/components/parameters/MediaFieldsParameter"
+          },
+          {
+            "$ref": "#/components/parameters/PlaceFieldsParameter"
+          },
+          {
+            "$ref": "#/components/parameters/PollFieldsParameter"
+          }
+        ],
+        "responses": {
+          "200": {
+            "content": {
+              "application/json": {
+                "schema": {
+                  "properties": {
+                    "data": {
+                      "items": {
+                        "$ref": "#/components/schemas/Tweet"
+                      },
+                      "minItems": 1,
+                      "type": "array"
+                    },
+                    "errors": {
+                      "items": {
+                        "$ref": "#/components/schemas/Problem"
+                      },
+                      "minItems": 1,
+                      "type": "array"
+                    },
+                    "includes": {
+                      "$ref": "#/components/schemas/Expansions"
+                    },
+                    "meta": {
+                      "properties": {
+                        "next_token": {
+                          "description": "The next token",
+                          "pattern": "^[a-zA-Z0-9]+$",
+                          "type": "string"
+                        },
+                        "previous_token": {
+                          "description": "The previous token",
+                          "pattern": "^[a-zA-Z0-9]+$",
+                          "type": "string"
+                        },
+                        "result_count": {
+                          "description": "Number of Tweets in result set",
+                          "format": "int32",
+                          "type": "integer"
+                        }
+                      },
+                      "type": "object"
+                    }
+                  },
+                  "type": "object"
+                }
+              }
+            },
+            "description": "The request was successful"
+          },
+          "default": {
+            "$ref": "#/components/responses/HttpErrorResponse"
+          }
+        },
+        "security": [
+          {
+            "BearerToken": [],
+            "UserToken": []
+          }
+        ],
+        "summary": "Returns Tweet objects liked by the provided User ID",
+        "tags": [
+          "Tweets"
+        ]
+      }
+    },
+    "/2/users/{id}/likes": {
+      "post": {
+        "description": "Causes the user (in the path) to like the specified tweet. The user in the path must match the user context authorizing the request.",
+        "operationId": "usersIdLike",
+        "parameters": [
+          {
+            "description": "The ID of the user that is requesting to like the tweet",
+            "in": "path",
+            "name": "id",
+            "required": true,
+            "schema": {
+              "$ref": "#/components/schemas/UserID"
+            }
+          }
+        ],
+        "requestBody": {
+          "content": {
+            "application/json": {
+              "schema": {
+                "$ref": "#/components/schemas/UsersLikesCreateRequest"
+              }
+            }
+          }
+        },
+        "responses": {
+          "200": {
+            "content": {
+              "application/json": {
+                "schema": {
+                  "$ref": "#/components/schemas/UsersLikesCreateResponse"
+                }
+              }
+            },
+            "description": "The request was successful"
+          },
+          "default": {
+            "$ref": "#/components/responses/HttpErrorResponse"
+          }
+        },
+        "security": [
+          {
+            "UserToken": []
+          }
+        ],
+        "summary": "Causes the user (in the path) to like the specified tweet",
+        "tags": [
+          "Tweets"
+        ]
+      }
+    },
+    "/2/users/{id}/likes/{tweet_id}": {
+      "delete": {
+        "description": "Causes the user (in the path) to unlike the specified tweet. The user must match the user context authorizing the request",
+        "operationId": "usersIdUnlike",
+        "parameters": [
+          {
+            "description": "The ID of the user that is requesting to unlike the tweet",
+            "in": "path",
+            "name": "id",
+            "required": true,
+            "schema": {
+              "$ref": "#/components/schemas/UserID"
+            }
+          },
+          {
+            "description": "The ID of the tweet that the user is requesting to unlike",
+            "in": "path",
+            "name": "tweet_id",
+            "required": true,
+            "schema": {
+              "$ref": "#/components/schemas/TweetID"
+            }
+          }
+        ],
+        "responses": {
+          "200": {
+            "content": {
+              "application/json": {
+                "schema": {
+                  "$ref": "#/components/schemas/UsersLikesDeleteResponse"
+                }
+              }
+            },
+            "description": "The request was successful"
+          },
+          "default": {
+            "$ref": "#/components/responses/HttpErrorResponse"
+          }
+        },
+        "security": [
+          {
+            "UserToken": []
+          }
+        ],
+        "summary": "Causes the user (in the path) to unlike the specified tweet",
+        "tags": [
+          "Tweets"
+        ]
+      }
+    },
+    "/2/users/{id}/mentions": {
+      "get": {
+        "description": "Returns Tweet objects that mention username associated to the provided User ID",
+        "operationId": "usersIdMentions",
+        "parameters": [
+          {
+            "description": "The ID of the User to list mentions of",
+            "in": "path",
+            "name": "id",
+            "required": true,
+            "schema": {
+              "$ref": "#/components/schemas/UserID"
+            }
+          },
+          {
+            "$ref": "#/components/parameters/SinceIdRequestParameter"
+          },
+          {
+            "$ref": "#/components/parameters/UntilIdRequestParameter"
+          },
+          {
+            "$ref": "#/components/parameters/MaxResultsRequestParameter"
+          },
+          {
+            "$ref": "#/components/parameters/PaginationTokenRequestParameter"
+          },
+          {
+            "$ref": "#/components/parameters/StartTimeRequestParameter"
+          },
+          {
+            "$ref": "#/components/parameters/EndTimeRequestParameter"
+          },
+          {
+            "$ref": "#/components/parameters/TweetExpansionsParameter"
+          },
+          {
+            "$ref": "#/components/parameters/TweetFieldsParameter"
+          },
+          {
+            "$ref": "#/components/parameters/UserFieldsParameter"
+          },
+          {
+            "$ref": "#/components/parameters/MediaFieldsParameter"
+          },
+          {
+            "$ref": "#/components/parameters/PlaceFieldsParameter"
+          },
+          {
+            "$ref": "#/components/parameters/PollFieldsParameter"
+          }
+        ],
+        "responses": {
+          "200": {
+            "content": {
+              "application/json": {
+                "schema": {
+                  "$ref": "#/components/schemas/GenericTweetsTimelineResponse"
+                }
+              }
+            },
+            "description": "The request was successful"
+          },
+          "default": {
+            "$ref": "#/components/responses/HttpErrorResponse"
+          }
+        },
+        "security": [
+          {
+            "BearerToken": [],
+            "UserToken": []
+          }
+        ],
+        "summary": "User mention timeline by User ID",
+        "tags": [
+          "Tweets"
+        ]
+      }
+    },
+    "/2/users/{id}/muting": {
+      "post": {
+        "description": "Causes the user (in the path) to mute the target user. The user (in the path) must match the user context authorizing the request",
+        "operationId": "usersIdMute",
+        "parameters": [
+          {
+            "description": "The ID of the user that is requesting to mute the target user",
+            "in": "path",
+            "name": "id",
+            "required": true,
+            "schema": {
+              "$ref": "#/components/schemas/UserID"
+            }
+          }
+        ],
+        "requestBody": {
+          "content": {
+            "application/json": {
+              "schema": {
+                "properties": {
+                  "target_user_id": {
+                    "pattern": "^[0-9]{1,19}$",
+                    "type": "string"
+                  }
+                },
+                "required": [
+                  "target_user_id"
+                ],
+                "type": "object"
+              }
+            }
+          }
+        },
+        "responses": {
+          "200": {
+            "content": {
+              "application/json": {
+                "schema": {
+                  "$ref": "#/components/schemas/UsersMutingMutationResponse"
+                }
+              }
+            },
+            "description": "The request was successful"
+          },
+          "default": {
+            "$ref": "#/components/responses/HttpErrorResponse"
+          }
+        },
+        "security": [
+          {
+            "UserToken": []
+          }
+        ],
+        "summary": "Mute User by User ID",
+        "tags": [
+          "Users"
+        ]
+      }
+    },
+    "/2/users/{id}/retweets": {
+      "post": {
+        "description": "Causes the user (in the path) to retweet the specified tweet. The user in the path must match the user context authorizing the request.",
+        "operationId": "usersIdRetweets",
+        "parameters": [
+          {
+            "description": "The ID of the user that is requesting to retweet the tweet",
+            "in": "path",
+            "name": "id",
+            "required": true,
+            "schema": {
+              "$ref": "#/components/schemas/UserID"
+            }
+          }
+        ],
+        "requestBody": {
+          "content": {
+            "application/json": {
+              "schema": {
+                "$ref": "#/components/schemas/UsersRetweetsCreateRequest"
+              }
+            }
+          }
+        },
+        "responses": {
+          "200": {
+            "content": {
+              "application/json": {
+                "schema": {
+                  "$ref": "#/components/schemas/UsersRetweetsCreateResponse"
+                }
+              }
+            },
+            "description": "The request was successful"
+          },
+          "default": {
+            "$ref": "#/components/responses/HttpErrorResponse"
+          }
+        },
+        "security": [
+          {
+            "UserToken": []
+          }
+        ],
+        "summary": "Causes the user (in the path) to retweet the specified tweet",
+        "tags": [
+          "Tweets"
+        ]
+      }
+    },
+    "/2/users/{id}/retweets/{source_tweet_id}": {
+      "delete": {
+        "description": "Causes the user (in the path) to unretweet the specified tweet. The user must match the user context authorizing the request",
+        "operationId": "usersIdUnretweets",
+        "parameters": [
+          {
+            "description": "The ID of the user that is requesting to unretweet the tweet",
+            "in": "path",
+            "name": "id",
+            "required": true,
+            "schema": {
+              "$ref": "#/components/schemas/UserID"
+            }
+          },
+          {
+            "description": "The ID of the tweet that the user is requesting to unretweet",
+            "in": "path",
+            "name": "source_tweet_id",
+            "required": true,
+            "schema": {
+              "$ref": "#/components/schemas/TweetID"
+            }
+          }
+        ],
+        "responses": {
+          "200": {
+            "content": {
+              "application/json": {
+                "schema": {
+                  "$ref": "#/components/schemas/UsersRetweetsDeleteResponse"
+                }
+              }
+            },
+            "description": "The request was successful"
+          },
+          "default": {
+            "$ref": "#/components/responses/HttpErrorResponse"
+          }
+        },
+        "security": [
+          {
+            "UserToken": []
+          }
+        ],
+        "summary": "Causes the user (in the path) to unretweet the specified tweet",
+        "tags": [
+          "Tweets"
+        ]
+      }
+    },
+    "/2/users/{id}/tweets": {
+      "get": {
+        "description": "Returns a list of Tweets authored by the provided User ID",
+        "operationId": "usersIdTweets",
+        "parameters": [
+          {
+            "description": "The ID of the User to list Tweets of",
+            "in": "path",
+            "name": "id",
+            "required": true,
+            "schema": {
+              "$ref": "#/components/schemas/UserID"
+            }
+          },
+          {
+            "$ref": "#/components/parameters/SinceIdRequestParameter"
+          },
+          {
+            "$ref": "#/components/parameters/UntilIdRequestParameter"
+          },
+          {
+            "$ref": "#/components/parameters/MaxResultsRequestParameter"
+          },
+          {
+            "$ref": "#/components/parameters/TweetTypeExcludesRequestParameter"
+          },
+          {
+            "$ref": "#/components/parameters/PaginationTokenRequestParameter"
+          },
+          {
+            "$ref": "#/components/parameters/StartTimeRequestParameter"
+          },
+          {
+            "$ref": "#/components/parameters/EndTimeRequestParameter"
+          },
+          {
+            "$ref": "#/components/parameters/TweetExpansionsParameter"
+          },
+          {
+            "$ref": "#/components/parameters/TweetFieldsParameter"
+          },
+          {
+            "$ref": "#/components/parameters/UserFieldsParameter"
+          },
+          {
+            "$ref": "#/components/parameters/MediaFieldsParameter"
+          },
+          {
+            "$ref": "#/components/parameters/PlaceFieldsParameter"
+          },
+          {
+            "$ref": "#/components/parameters/PollFieldsParameter"
+          }
+        ],
+        "responses": {
+          "200": {
+            "content": {
+              "application/json": {
+                "schema": {
+                  "$ref": "#/components/schemas/GenericTweetsTimelineResponse"
+                }
+              }
+            },
+            "description": "The request was successful"
+          },
+          "default": {
+            "$ref": "#/components/responses/HttpErrorResponse"
+          }
+        },
+        "security": [
+          {
+            "BearerToken": [],
+            "UserToken": []
+          }
+        ],
+        "summary": "User Tweets timeline by User ID",
+        "tags": [
+          "Tweets"
+        ]
+      }
+    },
+    "/2/users/{source_user_id}/blocking/{target_user_id}": {
+      "delete": {
+        "description": "Causes the source user to unblock the target user. The source user must match the user context authorizing the request",
+        "operationId": "usersIdUnblock",
+        "parameters": [
+          {
+            "description": "The ID of the user that is requesting to unblock the target user",
+            "in": "path",
+            "name": "source_user_id",
+            "required": true,
+            "schema": {
+              "$ref": "#/components/schemas/UserID"
+            }
+          },
+          {
+            "description": "The ID of the user that the source user is requesting to unblock",
+            "in": "path",
+            "name": "target_user_id",
+            "required": true,
+            "schema": {
+              "$ref": "#/components/schemas/UserID"
+            }
+          }
+        ],
+        "responses": {
+          "200": {
+            "content": {
+              "application/json": {
+                "schema": {
+                  "$ref": "#/components/schemas/UsersBlockingMutationResponse"
+                }
+              }
+            },
+            "description": "The request was successful"
+          },
+          "default": {
+            "$ref": "#/components/responses/HttpErrorResponse"
+          }
+        },
+        "security": [
+          {
+            "UserToken": []
+          }
+        ],
+        "summary": "Unblock User by User ID",
+        "tags": [
+          "Users"
+        ]
+      }
+    },
+    "/2/users/{source_user_id}/following/{target_user_id}": {
+      "delete": {
+        "description": "Causes the source user to unfollow the target user. The source user must match the user context authorizing the request",
+        "operationId": "usersIdUnfollow",
+        "parameters": [
+          {
+            "description": "The ID of the user that is requesting to unfollow the target user",
+            "in": "path",
+            "name": "source_user_id",
+            "required": true,
+            "schema": {
+              "$ref": "#/components/schemas/UserID"
+            }
+          },
+          {
+            "description": "The ID of the user that the source user is requesting to unfollow",
+            "in": "path",
+            "name": "target_user_id",
+            "required": true,
+            "schema": {
+              "$ref": "#/components/schemas/UserID"
+            }
+          }
+        ],
+        "responses": {
+          "200": {
+            "content": {
+              "application/json": {
+                "schema": {
+                  "$ref": "#/components/schemas/UsersFollowingDeleteResponse"
+                }
+              }
+            },
+            "description": "The request was successful"
+          },
+          "default": {
+            "$ref": "#/components/responses/HttpErrorResponse"
+          }
+        },
+        "security": [
+          {
+            "UserToken": []
+          }
+        ],
+        "summary": "Unfollow User",
+        "tags": [
+          "Users"
+        ]
+      }
+    },
+    "/2/users/{source_user_id}/muting/{target_user_id}": {
+      "delete": {
+        "description": "Causes the source user to unmute the target user. The source user must match the user context authorizing the request",
+        "operationId": "usersIdUnmute",
+        "parameters": [
+          {
+            "description": "The ID of the user that is requesting to unmute the target user",
+            "in": "path",
+            "name": "source_user_id",
+            "required": true,
+            "schema": {
+              "$ref": "#/components/schemas/UserID"
+            }
+          },
+          {
+            "description": "The ID of the user that the source user is requesting to unmute",
+            "in": "path",
+            "name": "target_user_id",
+            "required": true,
+            "schema": {
+              "$ref": "#/components/schemas/UserID"
+            }
+          }
+        ],
+        "responses": {
+          "200": {
+            "content": {
+              "application/json": {
+                "schema": {
+                  "$ref": "#/components/schemas/UsersMutingMutationResponse"
+                }
+              }
+            },
+            "description": "The request was successful"
+          },
+          "default": {
+            "$ref": "#/components/responses/HttpErrorResponse"
+          }
+        },
+        "security": [
+          {
+            "UserToken": []
+          }
+        ],
+        "summary": "Unmute User by User ID",
+        "tags": [
+          "Users"
+        ]
+      }
+    }
+  },
+  "components": {
+    "parameters": {
+      "BackfillMinutesRequestParameter": {
+        "description": "The number of minutes of backfill requested",
+        "in": "query",
+        "name": "backfill_minutes",
+        "required": false,
+        "schema": {
+          "format": "int32",
+          "maximum": 5,
+          "minimum": 0,
+          "type": "integer"
+        }
+      },
+      "EndTimeRequestParameter": {
+        "description": "YYYY-MM-DDTHH:mm:ssZ. The latest UTC timestamp to which the Tweets will be provided. The until_id parameter takes precedence if it is also specified.",
+        "example": "2021-02-14T18:40:40.000Z",
+        "in": "query",
+        "name": "end_time",
+        "required": false,
+        "schema": {
+          "format": "date-time",
+          "type": "string"
+        }
+      },
+      "MaxResultsRequestParameter": {
+        "description": "The maximum number of results",
+        "in": "query",
+        "name": "max_results",
+        "required": false,
+        "schema": {
+          "format": "int32",
+          "maximum": 100,
+          "minimum": 5,
+          "type": "integer"
+        }
+      },
+      "MediaFieldsParameter": {
+        "description": "A comma separated list of Media fields to display.",
+        "explode": false,
+        "in": "query",
+        "name": "media.fields",
+        "required": false,
+        "schema": {
+          "description": "The list of fields you can display for a Media object.",
+          "example": [
+            "media_key",
+            "duration_ms",
+            "height",
+            "preview_image_url",
+            "type",
+            "url",
+            "width",
+            "public_metrics"
+          ],
+          "items": {
+            "enum": [
+              "media_key",
+              "duration_ms",
+              "height",
+              "preview_image_url",
+              "type",
+              "url",
+              "width",
+              "public_metrics",
+              "non_public_metrics",
+              "organic_metrics",
+              "promoted_metrics",
+              "alt_text"
+            ],
+            "type": "string"
+          },
+          "minItems": 1,
+          "type": "array",
+          "uniqueItems": true
+        },
+        "style": "form"
+      },
+      "PaginationTokenRequestParameter": {
+        "description": "This parameter is used to get the next 'page' of results.",
+        "in": "query",
+        "name": "pagination_token",
+        "required": false,
+        "schema": {
+          "minLength": 1,
+          "type": "string"
+        }
+      },
+      "PlaceFieldsParameter": {
+        "description": "A comma separated list of Place fields to display.",
+        "explode": false,
+        "in": "query",
+        "name": "place.fields",
+        "required": false,
+        "schema": {
+          "description": "The list of fields you can display for a Place object.",
+          "example": [
+            "id",
+            "name",
+            "country_code",
+            "place_type",
+            "full_name",
+            "country",
+            "contained_within",
+            "geo"
+          ],
+          "items": {
+            "enum": [
+              "id",
+              "name",
+              "country_code",
+              "place_type",
+              "full_name",
+              "country",
+              "contained_within",
+              "geo"
+            ],
+            "type": "string"
+          },
+          "minItems": 1,
+          "type": "array",
+          "uniqueItems": true
+        },
+        "style": "form"
+      },
+      "PollFieldsParameter": {
+        "description": "A comma separated list of Poll fields to display.",
+        "explode": false,
+        "in": "query",
+        "name": "poll.fields",
+        "required": false,
+        "schema": {
+          "description": "The list of fields you can display for a Poll object.",
+          "example": [
+            "duration_minutes",
+            "options",
+            "voting_status",
+            "end_datetime"
+          ],
+          "items": {
+            "enum": [
+              "id",
+              "options",
+              "voting_status",
+              "end_datetime",
+              "duration_minutes"
+            ],
+            "type": "string"
+          },
+          "minItems": 1,
+          "type": "array",
+          "uniqueItems": true
+        },
+        "style": "form"
+      },
+      "SinceIdRequestParameter": {
+        "description": "The minimum Tweet ID to be included in the result set. This parameter takes precedence over start_time if both are specified.",
+        "example": "791775337160081409",
+        "in": "query",
+        "name": "since_id",
+        "required": false,
+        "schema": {
+          "$ref": "#/components/schemas/TweetID"
+        }
+      },
+      "SpaceExpansionsParameter": {
+        "description": "A comma separated list of fields to expand.",
+        "explode": false,
+        "in": "query",
+        "name": "expansions",
+        "schema": {
+          "description": "The list of fields you can expand for a [Space](#Space) object. If the field has an ID, it can be expanded into a full object.",
+          "example": [
+            "creator_id",
+            "host_ids",
+            "invited_user_ids",
+            "speaker_ids"
+          ],
+          "items": {
+            "enum": [
+              "creator_id",
+              "host_ids",
+              "invited_user_ids",
+              "speaker_ids"
+            ],
+            "type": "string"
+          },
+          "minItems": 1,
+          "type": "array",
+          "uniqueItems": true
+        },
+        "style": "form"
+      },
+      "SpaceFieldsParameter": {
+        "description": "A comma separated list of Space fields to display.",
+        "explode": false,
+        "in": "query",
+        "name": "space.fields",
+        "required": false,
+        "schema": {
+          "description": "The fields available for a Space object.",
+          "example": [
+            "created_at",
+            "creator_id",
+            "host_ids",
+            "invited_user_ids",
+            "is_ticketed",
+            "lang",
+            "participant_count",
+            "scheduled_start",
+            "speaker_ids",
+            "started_at",
+            "title",
+            "updated_at"
+          ],
+          "items": {
+            "enum": [
+              "created_at",
+              "creator_id",
+              "host_ids",
+              "invited_user_ids",
+              "is_ticketed",
+              "lang",
+              "participant_count",
+              "scheduled_start",
+              "speaker_ids",
+              "started_at",
+              "title",
+              "updated_at"
+            ],
+            "type": "string"
+          },
+          "minItems": 1,
+          "type": "array",
+          "uniqueItems": true
+        },
+        "style": "form"
+      },
+      "StartTimeRequestParameter": {
+        "description": "YYYY-MM-DDTHH:mm:ssZ. The earliest UTC timestamp from which the Tweets will be provided. The since_id parameter takes precedence if it is also specified.",
+        "example": "2021-02-01T18:40:40.000Z",
+        "in": "query",
+        "name": "start_time",
+        "required": false,
+        "schema": {
+          "format": "date-time",
+          "type": "string"
+        }
+      },
+      "TweetExpansionsParameter": {
+        "description": "A comma separated list of fields to expand.",
+        "explode": false,
+        "in": "query",
+        "name": "expansions",
+        "schema": {
+          "description": "The list of fields you can expand for a [Tweet](#Tweet) object. If the field has an ID, it can be expanded into a full object.",
+          "example": [
+            "author_id",
+            "referenced_tweets.id",
+            "in_reply_to_user_id",
+            "geo.place_id",
+            "attachments.media_keys",
+            "attachments.poll_ids",
+            "entities.mentions.username",
+            "referenced_tweets.id.author_id"
+          ],
+          "items": {
+            "enum": [
+              "author_id",
+              "referenced_tweets.id",
+              "in_reply_to_user_id",
+              "geo.place_id",
+              "attachments.media_keys",
+              "attachments.poll_ids",
+              "entities.mentions.username",
+              "referenced_tweets.id.author_id"
+            ],
+            "type": "string"
+          },
+          "minItems": 1,
+          "type": "array",
+          "uniqueItems": true
+        },
+        "style": "form"
+      },
+      "TweetFieldsParameter": {
+        "description": "A comma separated list of Tweet fields to display.",
+        "explode": false,
+        "in": "query",
+        "name": "tweet.fields",
+        "required": false,
+        "schema": {
+          "description": "The list of fields you can display for a [Tweet](#Tweet) object.",
+          "example": [
+            "created_at",
+            "author_id",
+            "entities",
+            "conversation_id",
+            "reply_settings",
+            "public_metrics"
+          ],
+          "items": {
+            "enum": [
+              "id",
+              "created_at",
+              "text",
+              "author_id",
+              "in_reply_to_user_id",
+              "referenced_tweets",
+              "attachments",
+              "withheld",
+              "geo",
+              "entities",
+              "public_metrics",
+              "possibly_sensitive",
+              "source",
+              "lang",
+              "context_annotations",
+              "non_public_metrics",
+              "promoted_metrics",
+              "organic_metrics",
+              "conversation_id",
+              "reply_settings"
+            ],
+            "type": "string"
+          },
+          "minItems": 1,
+          "type": "array",
+          "uniqueItems": true
+        },
+        "style": "form"
+      },
+      "TweetTypeExcludesRequestParameter": {
+        "description": "The set of entities to exclude (e.g. 'replies' or 'retweets')",
+        "explode": false,
+        "in": "query",
+        "name": "exclude",
+        "required": false,
+        "schema": {
+          "description": "The set of entities to exclude (e.g. 'replies' or 'retweets')",
+          "example": [
+            "replies",
+            "retweets"
+          ],
+          "items": {
+            "enum": [
+              "replies",
+              "retweets"
+            ],
+            "type": "string"
+          },
+          "minItems": 1,
+          "type": "array",
+          "uniqueItems": true
+        },
+        "style": "form"
+      },
+      "UntilIdRequestParameter": {
+        "description": "The maximum Tweet ID to be included in the result set. This parameter takes precedence over end_time if both are specified.",
+        "example": "1346889436626259968",
+        "in": "query",
+        "name": "until_id",
+        "required": false,
+        "schema": {
+          "$ref": "#/components/schemas/TweetID"
+        }
+      },
+      "UserExpansionsParameter": {
+        "description": "A comma separated list of fields to expand.",
+        "explode": false,
+        "in": "query",
+        "name": "expansions",
+        "schema": {
+          "description": "The list of fields you can expand for a [User](#User) object. If the field has an ID, it can be expanded into a full object.",
+          "example": [
+            "pinned_tweet_id"
+          ],
+          "items": {
+            "enum": [
+              "pinned_tweet_id"
+            ],
+            "type": "string"
+          },
+          "minItems": 1,
+          "type": "array",
+          "uniqueItems": true
+        },
+        "style": "form"
+      },
+      "UserFieldsParameter": {
+        "description": "A comma separated list of User fields to display.",
+        "explode": false,
+        "in": "query",
+        "name": "user.fields",
+        "required": false,
+        "schema": {
+          "description": "The list of fields you can display for a [User](#User) object.",
+          "example": [
+            "username",
+            "verified",
+            "profile_image_url"
+          ],
+          "items": {
+            "enum": [
+              "id",
+              "created_at",
+              "name",
+              "username",
+              "protected",
+              "verified",
+              "withheld",
+              "profile_image_url",
+              "location",
+              "url",
+              "description",
+              "entities",
+              "pinned_tweet_id",
+              "public_metrics"
+            ],
+            "type": "string"
+          },
+          "minItems": 1,
+          "type": "array",
+          "uniqueItems": true
+        },
+        "style": "form"
+      }
+    },
+    "responses": {
+      "HttpErrorResponse": {
+        "content": {
+          "application/json": {
+            "schema": {
+              "$ref": "#/components/schemas/Error"
+            }
+          },
+          "application/problem+json": {
+            "schema": {
+              "$ref": "#/components/schemas/Problem"
+            }
+          }
+        },
+        "description": "The request has failed."
+      }
+    },
+    "schemas": {
+      "AddOrDeleteRulesRequest": {
+        "oneOf": [
+          {
+            "$ref": "#/components/schemas/AddRulesRequest"
+          },
+          {
+            "$ref": "#/components/schemas/DeleteRulesRequest"
+          }
+        ],
+        "type": "object"
+      },
+      "AddOrDeleteRulesResponse": {
+        "oneOf": [
+          {
+            "$ref": "#/components/schemas/AddRulesResponse"
+          },
+          {
+            "$ref": "#/components/schemas/DeleteRulesResponse"
+          }
+        ],
+        "type": "object"
+      },
+      "AddRulesRequest": {
+        "description": "A request to add a user-specified stream filtering rule.",
+        "properties": {
+          "add": {
+            "items": {
+              "$ref": "#/components/schemas/RuleNoId"
+            },
+            "type": "array"
+          }
+        },
+        "required": [
+          "add"
+        ],
+        "type": "object"
+      },
+      "AddRulesResponse": {
+        "description": "A response from adding user-specified stream filtering rules.",
+        "properties": {
+          "data": {
+            "description": "All user-specified stream filtering rules that were created.",
+            "items": {
+              "$ref": "#/components/schemas/Rule"
+            },
+            "type": "array"
+          },
+          "meta": {
+            "$ref": "#/components/schemas/RulesResponseMetadata"
+          }
+        },
+        "required": [
+          "data",
+          "meta"
+        ],
+        "type": "object"
+      },
+      "AnimatedGif": {
+        "allOf": [
+          {
+            "$ref": "#/components/schemas/CommonMediaFields"
+          },
+          {
+            "properties": {
+              "preview_image_url": {
+                "format": "uri",
+                "type": "string"
+              },
+              "type": {
+                "enum": [
+                  "animated_gif"
+                ],
+                "type": "string"
+              }
+            },
+            "type": "object"
+          }
+        ],
+        "type": "object"
+      },
+      "CashtagEntity": {
+        "allOf": [
+          {
+            "$ref": "#/components/schemas/EntityIndicesInclusiveExclusive"
+          },
+          {
+            "$ref": "#/components/schemas/CashtagFields"
+          }
+        ]
+      },
+      "CashtagFields": {
+        "description": "Represent the portion of text recognized as a Cashtag, and its start and end position within the text.",
+        "properties": {
+          "tag": {
+            "example": "TWTR",
+            "type": "string"
+          }
+        },
+        "required": [
+          "tag"
+        ],
+        "type": "object"
+      },
+      "ClientDisconnectedProblem": {
+        "allOf": [
+          {
+            "$ref": "#/components/schemas/ProblemFields"
+          },
+          {
+            "properties": {
+              "type": {
+                "enum": [
+                  "https://api.twitter.com/2/problems/client-disconnected"
+                ],
+                "type": "string"
+              }
+            }
+          }
+        ],
+        "description": "Your client has gone away.",
+        "type": "object"
+      },
+      "ClientForbiddenProblem": {
+        "allOf": [
+          {
+            "$ref": "#/components/schemas/ProblemFields"
+          },
+          {
+            "properties": {
+              "reason": {
+                "enum": [
+                  "official-client-forbidden",
+                  "client-not-enrolled"
+                ],
+                "type": "string"
+              },
+              "registration_url": {
+                "format": "uri",
+                "type": "string"
+              },
+              "type": {
+                "enum": [
+                  "https://api.twitter.com/2/problems/client-forbidden"
+                ],
+                "type": "string"
+              }
+            },
+            "type": "object"
+          }
+        ],
+        "description": "A problem that indicates your client is forbidden from making this request.",
+        "type": "object"
+      },
+      "CommonMediaFields": {
+        "properties": {
+          "height": {
+            "$ref": "#/components/schemas/MediaHeight"
+          },
+          "media_key": {
+            "$ref": "#/components/schemas/MediaKey"
+          },
+          "width": {
+            "$ref": "#/components/schemas/MediaWidth"
+          }
+        }
+      },
+      "ComplianceJob": {
+        "properties": {
+          "created_at": {
+            "$ref": "#/components/schemas/CreatedAt"
+          },
+          "download_expires_at": {
+            "$ref": "#/components/schemas/DownloadExpiration"
+          },
+          "download_url": {
+            "$ref": "#/components/schemas/DownloadUrl"
+          },
+          "id": {
+            "$ref": "#/components/schemas/JobId"
+          },
+          "name": {
+            "$ref": "#/components/schemas/ComplianceJobName"
+          },
+          "status": {
+            "$ref": "#/components/schemas/ComplianceJobStatus"
+          },
+          "type": {
+            "$ref": "#/components/schemas/ComplianceJobType"
+          },
+          "upload_expires_at": {
+            "$ref": "#/components/schemas/UploadExpiration"
+          },
+          "upload_url": {
+            "$ref": "#/components/schemas/UploadUrl"
+          }
+        },
+        "required": [
+          "id",
+          "type",
+          "created_at",
+          "upload_url",
+          "download_url",
+          "upload_expires_at",
+          "download_expires_at",
+          "status"
+        ],
+        "type": "object"
+      },
+      "ComplianceJobName": {
+        "description": "User-provided name for a compliance job",
+        "example": "my-job",
+        "maxLength": 64,
+        "type": "string"
+      },
+      "ComplianceJobStatus": {
+        "description": "Status of a compliance job",
+        "enum": [
+          "created",
+          "in_progress",
+          "failed",
+          "complete"
+        ],
+        "type": "string"
+      },
+      "ComplianceJobType": {
+        "description": "Type of compliance job to list.",
+        "enum": [
+          "tweets",
+          "users"
+        ],
+        "type": "string"
+      },
+      "ConflictProblem": {
+        "allOf": [
+          {
+            "$ref": "#/components/schemas/ProblemFields"
+          },
+          {
+            "properties": {
+              "type": {
+                "enum": [
+                  "https://api.twitter.com/2/problems/conflict"
+                ],
+                "type": "string"
+              }
+            }
+          }
+        ],
+        "description": "You cannot create a new job if one is already in progress.",
+        "type": "object"
+      },
+      "ConnectionExceptionProblem": {
+        "allOf": [
+          {
+            "$ref": "#/components/schemas/ProblemFields"
+          },
+          {
+            "properties": {
+              "connection_issue": {
+                "enum": [
+                  "TooManyConnections",
+                  "ProvisioningSubscription",
+                  "RuleConfigurationIssue"
+                ],
+                "type": "string"
+              },
+              "type": {
+                "enum": [
+                  "https://api.twitter.com/2/problems/streaming-connection"
+                ],
+                "type": "string"
+              }
+            }
+          }
+        ],
+        "description": "A problem that indicates something is wrong with the connection",
+        "type": "object"
+      },
+      "ContextAnnotation": {
+        "description": "Annotation inferred from the tweet text.",
+        "properties": {
+          "domain": {
+            "$ref": "#/components/schemas/ContextAnnotationDomainFields"
+          },
+          "entity": {
+            "$ref": "#/components/schemas/ContextAnnotationEntityFields"
+          }
+        },
+        "required": [
+          "domain",
+          "entity"
+        ],
+        "type": "object"
+      },
+      "ContextAnnotationDomainFields": {
+        "description": "Represents the data for the context annotation domain.",
+        "properties": {
+          "description": {
+            "description": "Description of the context annotation domain.",
+            "type": "string"
+          },
+          "id": {
+            "description": "The unique id for a context annotation domain.",
+            "pattern": "^[0-9]{1,19}$",
+            "type": "string"
+          },
+          "name": {
+            "description": "Name of the context annotation domain.",
+            "type": "string"
+          }
+        },
+        "required": [
+          "id"
+        ],
+        "type": "object"
+      },
+      "ContextAnnotationEntityFields": {
+        "description": "Represents the data for the context annotation entity.",
+        "properties": {
+          "description": {
+            "description": "Description of the context annotation entity.",
+            "type": "string"
+          },
+          "id": {
+            "description": "The unique id for a context annotation entity.",
+            "pattern": "^[0-9]{1,19}$",
+            "type": "string"
+          },
+          "name": {
+            "description": "Name of the context annotation entity.",
+            "type": "string"
+          }
+        },
+        "required": [
+          "id"
+        ],
+        "type": "object"
+      },
+      "CountryCode": {
+        "description": "A two-letter ISO 3166-1 alpha-2 country code",
+        "example": "US",
+        "pattern": "^[A-Z]{2}$",
+        "type": "string"
+      },
+      "CreatedAt": {
+        "description": "Creation time of the compliance job.",
+        "example": "2021-01-06T18:40:40.000Z",
+        "format": "date-time",
+        "type": "string"
+      },
+      "DeleteRulesRequest": {
+        "description": "A response from deleting user-specified stream filtering rules.",
+        "properties": {
+          "delete": {
+            "description": "IDs and values of all deleted user-specified stream filtering rules.",
+            "properties": {
+              "ids": {
+                "description": "IDs of all deleted user-specified stream filtering rules.",
+                "items": {
+                  "$ref": "#/components/schemas/RuleId"
+                },
+                "type": "array"
+              },
+              "values": {
+                "description": "Values of all deleted user-specified stream filtering rules.",
+                "items": {
+                  "$ref": "#/components/schemas/RuleValue"
+                },
+                "type": "array"
+              }
+            },
+            "type": "object"
+          }
+        },
+        "required": [
+          "delete"
+        ],
+        "type": "object"
+      },
+      "DeleteRulesResponse": {
+        "properties": {
+          "meta": {
+            "$ref": "#/components/schemas/RulesResponseMetadata"
+          }
+        },
+        "required": [
+          "meta"
+        ],
+        "type": "object"
+      },
+      "DisallowedResourceProblem": {
+        "allOf": [
+          {
+            "$ref": "#/components/schemas/ProblemFields"
+          },
+          {
+            "properties": {
+              "resource_id": {
+                "type": "string"
+              },
+              "resource_type": {
+                "enum": [
+                  "tweet",
+                  "media"
+                ],
+                "type": "string"
+              },
+              "section": {
+                "enum": [
+                  "data",
+                  "includes"
+                ],
+                "type": "string"
+              },
+              "type": {
+                "enum": [
+                  "https://api.twitter.com/2/problems/disallowed-resource"
+                ],
+                "type": "string"
+              }
+            },
+            "required": [
+              "resource_id",
+              "resource_type",
+              "section"
+            ],
+            "type": "object"
+          }
+        ],
+        "description": "A problem that indicates that the resource requested violates the precepts of this API.",
+        "type": "object"
+      },
+      "DownloadExpiration": {
+        "description": "Expiration time of the download URL",
+        "example": "2021-01-06T18:40:40.000Z",
+        "format": "date-time",
+        "type": "string"
+      },
+      "DownloadUrl": {
+        "description": "URL from which the user will retrieve their compliance results",
+        "format": "uri",
+        "type": "string"
+      },
+      "DuplicateRuleProblem": {
+        "allOf": [
+          {
+            "$ref": "#/components/schemas/ProblemFields"
+          },
+          {
+            "properties": {
+              "type": {
+                "enum": [
+                  "https://api.twitter.com/2/problems/duplicate-rules"
+                ],
+                "type": "string"
+              }
+            }
+          }
+        ],
+        "description": "The rule you have submitted is a duplicate.",
+        "type": "object"
+      },
+      "End": {
+        "description": "The end time of the bucket",
+        "format": "date-time",
+        "type": "string"
+      },
+      "EntityIndicesInclusiveExclusive": {
+        "description": "Represent a boundary range (start and end index) for a recognized entity (for example a hashtag or a mention). `start` must be smaller than `end`.  The start index is inclusive, the end index is exclusive.",
+        "properties": {
+          "end": {
+            "description": "Index (zero-based) at which position this entity ends.  The index is exclusive.",
+            "example": 61,
+            "minimum": 0,
+            "type": "integer"
+          },
+          "start": {
+            "description": "Index (zero-based) at which position this entity starts.  The index is inclusive.",
+            "example": 50,
+            "minimum": 0,
+            "type": "integer"
+          }
+        },
+        "required": [
+          "start",
+          "end"
+        ],
+        "type": "object"
+      },
+      "EntityIndicesInclusiveInclusive": {
+        "description": "Represent a boundary range (start and end index) for a recognized entity (for example a hashtag or a mention). `start` must be smaller than `end`.  The start index is inclusive, the end index is inclusive.",
+        "properties": {
+          "end": {
+            "description": "Index (zero-based) at which position this entity ends.  The index is inclusive.",
+            "example": 61,
+            "minimum": 0,
+            "type": "integer"
+          },
+          "start": {
+            "description": "Index (zero-based) at which position this entity starts.  The index is inclusive.",
+            "example": 50,
+            "minimum": 0,
+            "type": "integer"
+          }
+        },
+        "required": [
+          "start",
+          "end"
+        ],
+        "type": "object"
+      },
+      "Error": {
+        "properties": {
+          "code": {
+            "format": "int32",
+            "type": "integer"
+          },
+          "message": {
+            "type": "string"
+          }
+        },
+        "required": [
+          "code",
+          "message"
+        ]
+      },
+      "Expansions": {
+        "properties": {
+          "media": {
+            "items": {
+              "$ref": "#/components/schemas/Media"
+            },
+            "minItems": 1,
+            "type": "array"
+          },
+          "places": {
+            "items": {
+              "$ref": "#/components/schemas/Place"
+            },
+            "minItems": 1,
+            "type": "array"
+          },
+          "polls": {
+            "items": {
+              "$ref": "#/components/schemas/Poll"
+            },
+            "minItems": 1,
+            "type": "array"
+          },
+          "tweets": {
+            "items": {
+              "$ref": "#/components/schemas/Tweet"
+            },
+            "minItems": 1,
+            "type": "array"
+          },
+          "users": {
+            "items": {
+              "$ref": "#/components/schemas/User"
+            },
+            "minItems": 1,
+            "type": "array"
+          }
+        },
+        "type": "object"
+      },
+      "FieldUnauthorizedProblem": {
+        "allOf": [
+          {
+            "$ref": "#/components/schemas/ProblemFields"
+          },
+          {
+            "properties": {
+              "field": {
+                "type": "string"
+              },
+              "resource_type": {
+                "enum": [
+                  "tweet",
+                  "media"
+                ],
+                "type": "string"
+              },
+              "section": {
+                "enum": [
+                  "data",
+                  "includes"
+                ],
+                "type": "string"
+              },
+              "type": {
+                "enum": [
+                  "https://api.twitter.com/2/problems/not-authorized-for-field"
+                ],
+                "type": "string"
+              }
+            },
+            "required": [
+              "resource_type",
+              "field",
+              "section"
+            ],
+            "type": "object"
+          }
+        ],
+        "description": "A problem that indicates that you are not allowed to see a particular field on a Tweet, User, etc.",
+        "type": "object"
+      },
+      "FilteredStreamingTweet": {
+        "description": "A tweet or error that can be returned by the streaming tweet API",
+        "oneOf": [
+          {
+            "description": "The values returned with a successful streamed tweet. Includes the user provided rules that the tweet matched.",
+            "properties": {
+              "data": {
+                "$ref": "#/components/schemas/Tweet"
+              },
+              "includes": {
+                "$ref": "#/components/schemas/Expansions"
+              },
+              "matching_rules": {
+                "description": "The list of rules which matched the tweet",
+                "items": {
+                  "properties": {
+                    "id": {
+                      "$ref": "#/components/schemas/RuleId"
+                    },
+                    "tag": {
+                      "description": "The user-supplied tag assigned to the rule which matched",
+                      "type": "string"
+                    }
+                  },
+                  "required": [
+                    "id"
+                  ],
+                  "type": "object"
+                },
+                "type": "array"
+              }
+            },
+            "required": [
+              "data",
+              "matching_rules"
+            ],
+            "type": "object"
+          },
+          {
+            "properties": {
+              "errors": {
+                "items": {
+                  "$ref": "#/components/schemas/Problem"
+                },
+                "type": "array"
+              }
+            },
+            "required": [
+              "errors"
+            ],
+            "type": "object"
+          }
+        ],
+        "type": "object"
+      },
+      "FullTextEntities": {
+        "properties": {
+          "cashtags": {
+            "items": {
+              "$ref": "#/components/schemas/CashtagEntity"
+            },
+            "minItems": 1,
+            "type": "array"
+          },
+          "hashtags": {
+            "items": {
+              "$ref": "#/components/schemas/HashtagEntity"
+            },
+            "minItems": 1,
+            "type": "array"
+          },
+          "mentions": {
+            "items": {
+              "$ref": "#/components/schemas/MentionEntity"
+            },
+            "minItems": 1,
+            "type": "array"
+          },
+          "urls": {
+            "items": {
+              "$ref": "#/components/schemas/UrlEntity"
+            },
+            "minItems": 1,
+            "type": "array"
+          }
+        },
+        "type": "object"
+      },
+      "GenericMultipleUsersLookupResponse": {
+        "properties": {
+          "data": {
+            "items": {
+              "$ref": "#/components/schemas/User"
+            },
+            "type": "array"
+          },
+          "errors": {
+            "items": {
+              "$ref": "#/components/schemas/Problem"
+            },
+            "minItems": 1,
+            "type": "array"
+          },
+          "includes": {
+            "$ref": "#/components/schemas/Expansions"
+          },
+          "meta": {
+            "properties": {
+              "next_token": {
+                "description": "The next token",
+                "pattern": "^[a-zA-Z0-9]+$",
+                "type": "string"
+              },
+              "previous_token": {
+                "description": "The previous token",
+                "pattern": "^[a-zA-Z0-9]+$",
+                "type": "string"
+              },
+              "result_count": {
+                "description": "The number of user results returned in this response",
+                "type": "integer"
+              }
+            },
+            "type": "object"
+          }
+        },
+        "type": "object"
+      },
+      "GenericProblem": {
+        "allOf": [
+          {
+            "$ref": "#/components/schemas/ProblemFields"
+          },
+          {
+            "properties": {
+              "status": {
+                "type": "integer"
+              },
+              "type": {
+                "enum": [
+                  "about:blank"
+                ],
+                "type": "string"
+              }
+            },
+            "required": [
+              "status"
+            ],
+            "type": "object"
+          }
+        ],
+        "description": "A generic problem with no additional information beyond that provided by the HTTP status code."
+      },
+      "GenericTweetsTimelineResponse": {
+        "properties": {
+          "data": {
+            "items": {
+              "$ref": "#/components/schemas/Tweet"
+            },
+            "minItems": 1,
+            "type": "array"
+          },
+          "errors": {
+            "items": {
+              "$ref": "#/components/schemas/Problem"
+            },
+            "minItems": 1,
+            "type": "array"
+          },
+          "includes": {
+            "$ref": "#/components/schemas/Expansions"
+          },
+          "meta": {
+            "properties": {
+              "newest_id": {
+                "$ref": "#/components/schemas/TweetID",
+                "description": "Newest Tweet ID in the result set"
+              },
+              "next_token": {
+                "description": "The next token",
+                "pattern": "^[a-zA-Z0-9]+$",
+                "type": "string"
+              },
+              "oldest_id": {
+                "$ref": "#/components/schemas/TweetID",
+                "description": "Oldest Tweet ID in the result set"
+              },
+              "previous_token": {
+                "description": "The previous token",
+                "pattern": "^[a-zA-Z0-9]+$",
+                "type": "string"
+              },
+              "result_count": {
+                "description": "Number of Tweets in result set",
+                "format": "int32",
+                "type": "integer"
+              }
+            },
+            "type": "object"
+          }
+        },
+        "type": "object"
+      },
+      "Geo": {
+        "properties": {
+          "bbox": {
+            "example": [
+              -105.193475,
+              39.60973,
+              -105.053164,
+              39.761974
+            ],
+            "items": {
+              "format": "double",
+              "maximum": 180,
+              "minimum": -180,
+              "type": "number"
+            },
+            "maxItems": 4,
+            "minItems": 4,
+            "type": "array"
+          },
+          "geometry": {
+            "$ref": "#/components/schemas/Point"
+          },
+          "properties": {
+            "type": "object"
+          },
+          "type": {
+            "enum": [
+              "Feature"
+            ],
+            "type": "string"
+          }
+        },
+        "required": [
+          "type",
+          "bbox",
+          "properties"
+        ],
+        "type": "object"
+      },
+      "Granularity": {
+        "default": "hour",
+        "enum": [
+          "minute",
+          "hour",
+          "day"
+        ],
+        "type": "string"
+      },
+      "HTTPStatusCode": {
+        "description": "HTTP Status Code.",
+        "maximum": 599,
+        "minimum": 100,
+        "type": "integer"
+      },
+      "HashtagEntity": {
+        "allOf": [
+          {
+            "$ref": "#/components/schemas/EntityIndicesInclusiveExclusive"
+          },
+          {
+            "$ref": "#/components/schemas/HashtagFields"
+          }
+        ]
+      },
+      "HashtagFields": {
+        "description": "Represent the portion of text recognized as a Hashtag, and its start and end position within the text.",
+        "properties": {
+          "tag": {
+            "description": "The text of the Hashtag",
+            "example": "MondayMotivation",
+            "type": "string"
+          }
+        },
+        "required": [
+          "tag"
+        ],
+        "type": "object"
+      },
+      "InvalidRequestProblem": {
+        "allOf": [
+          {
+            "$ref": "#/components/schemas/ProblemFields"
+          },
+          {
+            "properties": {
+              "errors": {
+                "items": {
+                  "properties": {
+                    "message": {
+                      "type": "string"
+                    },
+                    "parameters": {
+                      "additionalProperties": {
+                        "items": {
+                          "type": "string"
+                        },
+                        "type": "array"
+                      },
+                      "type": "object"
+                    }
+                  },
+                  "type": "object"
+                },
+                "type": "array"
+              },
+              "type": {
+                "enum": [
+                  "https://api.twitter.com/2/problems/invalid-request"
+                ],
+                "type": "string"
+              }
+            },
+            "type": "object"
+          }
+        ],
+        "description": "A problem that indicates this request is invalid."
+      },
+      "InvalidRuleProblem": {
+        "allOf": [
+          {
+            "$ref": "#/components/schemas/ProblemFields"
+          },
+          {
+            "properties": {
+              "type": {
+                "enum": [
+                  "https://api.twitter.com/2/problems/invalid-rules"
+                ],
+                "type": "string"
+              }
+            }
+          }
+        ],
+        "description": "The rule you have submitted is invalid.",
+        "type": "object"
+      },
+      "JobId": {
+        "description": "Compliance Job ID",
+        "example": "1372966999991541762",
+        "pattern": "^[0-9]{1,19}$",
+        "type": "string"
+      },
+      "Media": {
+        "discriminator": {
+          "mapping": {
+            "animated_gif": "#/components/schemas/AnimatedGif",
+            "photo": "#/components/schemas/Photo",
+            "video": "#/components/schemas/Video"
+          },
+          "propertyName": "type"
+        },
+        "oneOf": [
+          {
+            "$ref": "#/components/schemas/Photo"
+          },
+          {
+            "$ref": "#/components/schemas/Video"
+          },
+          {
+            "$ref": "#/components/schemas/AnimatedGif"
+          }
+        ],
+        "properties": {
+          "type": {
+            "type": "string"
+          }
+        },
+        "required": [
+          "type"
+        ],
+        "type": "object"
+      },
+      "MediaHeight": {
+        "description": "The height of the media in pixels",
+        "minimum": 0,
+        "type": "integer"
+      },
+      "MediaKey": {
+        "description": "The Media Key identifier for this attachment.",
+        "pattern": "^([0-9]+)_([0-9]+)$",
+        "type": "string"
+      },
+      "MediaWidth": {
+        "description": "The width of the media in pixels",
+        "minimum": 0,
+        "type": "integer"
+      },
+      "MentionEntity": {
+        "allOf": [
+          {
+            "$ref": "#/components/schemas/EntityIndicesInclusiveExclusive"
+          },
+          {
+            "$ref": "#/components/schemas/MentionFields"
+          }
+        ]
+      },
+      "MentionFields": {
+        "description": "Represent the portion of text recognized as a User mention, and its start and end position within the text.",
+        "properties": {
+          "id": {
+            "$ref": "#/components/schemas/UserID"
+          },
+          "username": {
+            "$ref": "#/components/schemas/UserName"
+          }
+        },
+        "required": [
+          "username",
+          "id"
+        ],
+        "type": "object"
+      },
+      "MultiComplianceJobResponse": {
+        "properties": {
+          "data": {
+            "items": {
+              "$ref": "#/components/schemas/ComplianceJob"
+            },
+            "minItems": 0,
+            "type": "array"
+          },
+          "errors": {
+            "items": {
+              "$ref": "#/components/schemas/Problem"
+            },
+            "minItems": 0,
+            "type": "array"
+          }
+        },
+        "type": "object"
+      },
+      "MultiSpaceLookupResponse": {
+        "properties": {
+          "data": {
+            "items": {
+              "$ref": "#/components/schemas/Space"
+            },
+            "minItems": 0,
+            "type": "array"
+          },
+          "errors": {
+            "items": {
+              "$ref": "#/components/schemas/Problem"
+            },
+            "minItems": 1,
+            "type": "array"
+          },
+          "includes": {
+            "$ref": "#/components/schemas/Expansions"
+          }
+        },
+        "type": "object"
+      },
+      "MultiTweetLookupResponse": {
+        "properties": {
+          "data": {
+            "items": {
+              "$ref": "#/components/schemas/Tweet"
+            },
+            "minItems": 1,
+            "type": "array"
+          },
+          "errors": {
+            "items": {
+              "$ref": "#/components/schemas/Problem"
+            },
+            "minItems": 1,
+            "type": "array"
+          },
+          "includes": {
+            "$ref": "#/components/schemas/Expansions"
+          }
+        },
+        "type": "object"
+      },
+      "MultiUserLookupResponse": {
+        "properties": {
+          "data": {
+            "items": {
+              "$ref": "#/components/schemas/User"
+            },
+            "minItems": 1,
+            "type": "array"
+          },
+          "errors": {
+            "items": {
+              "$ref": "#/components/schemas/Problem"
+            },
+            "minItems": 1,
+            "type": "array"
+          },
+          "includes": {
+            "$ref": "#/components/schemas/Expansions"
+          }
+        },
+        "type": "object"
+      },
+      "OperationalDisconnectProblem": {
+        "allOf": [
+          {
+            "$ref": "#/components/schemas/ProblemFields"
+          },
+          {
+            "properties": {
+              "disconnect_type": {
+                "enum": [
+                  "OperationalDisconnect",
+                  "UpstreamOperationalDisconnect",
+                  "ForceDisconnect",
+                  "UpstreamUncleanDisconnect",
+                  "SlowReader",
+                  "InternalError",
+                  "PackageUpgraded",
+                  "PackageDowngraded",
+                  "ClientApplicationStateDegraded"
+                ],
+                "type": "string"
+              },
+              "type": {
+                "enum": [
+                  "https://api.twitter.com/2/problems/operational-disconnect"
+                ],
+                "type": "string"
+              }
+            }
+          }
+        ],
+        "description": "You have been disconnected for operational reasons.",
+        "type": "object"
+      },
+      "Photo": {
+        "allOf": [
+          {
+            "$ref": "#/components/schemas/CommonMediaFields"
+          },
+          {
+            "properties": {
+              "alt_text": {
+                "type": "string"
+              },
+              "type": {
+                "enum": [
+                  "photo"
+                ],
+                "type": "string"
+              },
+              "url": {
+                "format": "uri",
+                "type": "string"
+              }
+            },
+            "type": "object"
+          }
+        ],
+        "type": "object"
+      },
+      "Place": {
+        "properties": {
+          "contained_within": {
+            "items": {
+              "$ref": "#/components/schemas/PlaceId"
+            },
+            "minItems": 1,
+            "type": "array"
+          },
+          "country": {
+            "description": "The full name of the county in which this place exists.",
+            "example": "United States",
+            "type": "string"
+          },
+          "country_code": {
+            "$ref": "#/components/schemas/CountryCode"
+          },
+          "full_name": {
+            "description": "The full name of this place.",
+            "example": "Lakewood, CO",
+            "type": "string"
+          },
+          "geo": {
+            "$ref": "#/components/schemas/Geo"
+          },
+          "id": {
+            "$ref": "#/components/schemas/PlaceId"
+          },
+          "name": {
+            "description": "The human readable name of this place.",
+            "example": "Lakewood",
+            "type": "string"
+          },
+          "place_type": {
+            "$ref": "#/components/schemas/PlaceType"
+          }
+        },
+        "required": [
+          "id",
+          "full_name"
+        ]
+      },
+      "PlaceId": {
+        "description": "The identifier for this place",
+        "example": "f7eb2fa2fea288b1",
+        "type": "string"
+      },
+      "PlaceType": {
+        "enum": [
+          "poi",
+          "neighborhood",
+          "city",
+          "admin",
+          "country",
+          "unknown"
+        ],
+        "example": "city",
+        "type": "string"
+      },
+      "Point": {
+        "description": "A [GeoJson Point](https://tools.ietf.org/html/rfc7946#section-3.1.2) geometry object.",
+        "properties": {
+          "coordinates": {
+            "$ref": "#/components/schemas/Position"
+          },
+          "type": {
+            "enum": [
+              "Point"
+            ],
+            "example": "Point",
+            "type": "string"
+          }
+        },
+        "required": [
+          "type",
+          "coordinates"
+        ],
+        "type": "object"
+      },
+      "Poll": {
+        "description": "Represent a Poll attached to a Tweet",
+        "properties": {
+          "duration_minutes": {
+            "type": "integer"
+          },
+          "end_datetime": {
+            "format": "date-time",
+            "type": "string"
+          },
+          "id": {
+            "$ref": "#/components/schemas/PollId"
+          },
+          "options": {
+            "items": {
+              "$ref": "#/components/schemas/PollOption"
+            },
+            "maxItems": 4,
+            "minItems": 2,
+            "type": "array"
+          },
+          "voting_status": {
+            "enum": [
+              "open",
+              "closed"
+            ],
+            "type": "string"
+          }
+        },
+        "required": [
+          "id",
+          "options"
+        ],
+        "type": "object"
+      },
+      "PollId": {
+        "description": "Unique identifier of this poll.",
+        "example": "1365059861688410112",
+        "pattern": "^[0-9]{1,19}$",
+        "type": "string"
+      },
+      "PollOption": {
+        "description": "Describes a choice in a Poll object.",
+        "properties": {
+          "label": {
+            "description": "The text of a poll choice.",
+            "type": "string"
+          },
+          "position": {
+            "description": "Position of this choice in the poll.",
+            "type": "integer"
+          },
+          "votes": {
+            "description": "Number of users who voted for this choice.",
+            "type": "integer"
+          }
+        },
+        "required": [
+          "position",
+          "label",
+          "votes"
+        ],
+        "type": "object"
+      },
+      "Position": {
+        "description": "A [GeoJson Position](https://tools.ietf.org/html/rfc7946#section-3.1.1) in the format `[longitude,latitude]`.",
+        "example": [
+          -105.18816086351444,
+          40.247749999999996
+        ],
+        "items": {
+          "type": "number"
+        },
+        "maxItems": 2,
+        "minItems": 2,
+        "type": "array"
+      },
+      "Problem": {
+        "description": "An HTTP Problem Details object, as defined in IETF RFC 7807 (https://tools.ietf.org/html/rfc7807).",
+        "discriminator": {
+          "mapping": {
+            "about:blank": "#/components/schemas/GenericProblem",
+            "https://api.twitter.com/2/problems/client-disconnected": "#/components/schemas/ClientDisconnectedProblem",
+            "https://api.twitter.com/2/problems/client-forbidden": "#/components/schemas/ClientForbiddenProblem",
+            "https://api.twitter.com/2/problems/conflict": "#/components/schemas/ConflictProblem",
+            "https://api.twitter.com/2/problems/disallowed-resource": "#/components/schemas/DisallowedResourceProblem",
+            "https://api.twitter.com/2/problems/duplicate-rules": "#/components/schemas/DuplicateRuleProblem",
+            "https://api.twitter.com/2/problems/invalid-request": "#/components/schemas/InvalidRequestProblem",
+            "https://api.twitter.com/2/problems/invalid-rules": "#/components/schemas/InvalidRuleProblem",
+            "https://api.twitter.com/2/problems/not-authorized-for-field": "#/components/schemas/FieldUnauthorizedProblem",
+            "https://api.twitter.com/2/problems/not-authorized-for-resource": "#/components/schemas/ResourceUnauthorizedProblem",
+            "https://api.twitter.com/2/problems/operational-disconnect": "#/components/schemas/OperationalDisconnectProblem",
+            "https://api.twitter.com/2/problems/resource-not-found": "#/components/schemas/ResourceNotFoundProblem",
+            "https://api.twitter.com/2/problems/resource-unavailable": "#/components/schemas/ResourceUnavailableProblem",
+            "https://api.twitter.com/2/problems/rule-cap": "#/components/schemas/RulesCapProblem",
+            "https://api.twitter.com/2/problems/streaming-connection": "#/components/schemas/ConnectionExceptionProblem",
+            "https://api.twitter.com/2/problems/unsupported-authentication": "#/components/schemas/UnsupportedAuthenticationProblem",
+            "https://api.twitter.com/2/problems/usage-capped": "#/components/schemas/UsageCapExceededProblem"
+          },
+          "propertyName": "type"
+        },
+        "oneOf": [
+          {
+            "$ref": "#/components/schemas/GenericProblem"
+          },
+          {
+            "$ref": "#/components/schemas/InvalidRequestProblem"
+          },
+          {
+            "$ref": "#/components/schemas/ClientForbiddenProblem"
+          },
+          {
+            "$ref": "#/components/schemas/ResourceNotFoundProblem"
+          },
+          {
+            "$ref": "#/components/schemas/ResourceUnavailableProblem"
+          },
+          {
+            "$ref": "#/components/schemas/ResourceUnauthorizedProblem"
+          },
+          {
+            "$ref": "#/components/schemas/FieldUnauthorizedProblem"
+          },
+          {
+            "$ref": "#/components/schemas/DisallowedResourceProblem"
+          },
+          {
+            "$ref": "#/components/schemas/UnsupportedAuthenticationProblem"
+          },
+          {
+            "$ref": "#/components/schemas/UsageCapExceededProblem"
+          },
+          {
+            "$ref": "#/components/schemas/ConnectionExceptionProblem"
+          },
+          {
+            "$ref": "#/components/schemas/ClientDisconnectedProblem"
+          },
+          {
+            "$ref": "#/components/schemas/OperationalDisconnectProblem"
+          },
+          {
+            "$ref": "#/components/schemas/RulesCapProblem"
+          },
+          {
+            "$ref": "#/components/schemas/InvalidRuleProblem"
+          },
+          {
+            "$ref": "#/components/schemas/DuplicateRuleProblem"
+          },
+          {
+            "$ref": "#/components/schemas/ConflictProblem"
+          }
+        ],
+        "properties": {
+          "type": {
+            "type": "string"
+          }
+        },
+        "required": [
+          "type"
+        ]
+      },
+      "ProblemFields": {
+        "properties": {
+          "detail": {
+            "type": "string"
+          },
+          "title": {
+            "type": "string"
+          }
+        },
+        "required": [
+          "title",
+          "detail"
+        ],
+        "type": "object"
+      },
+      "ReplySettings": {
+        "description": "Shows who can reply a Tweet. Fields returned are everyone, mentioned_users, and following.",
+        "enum": [
+          "everyone",
+          "mentionedUsers",
+          "following",
+          "other"
+        ],
+        "pattern": "^[A-Za-z]{1,12}$",
+        "type": "string"
+      },
+      "ResourceNotFoundProblem": {
+        "allOf": [
+          {
+            "$ref": "#/components/schemas/ProblemFields"
+          },
+          {
+            "properties": {
+              "parameter": {
+                "minLength": 1,
+                "type": "string"
+              },
+              "resource_id": {
+                "type": "string"
+              },
+              "resource_type": {
+                "enum": [
+                  "user",
+                  "tweet",
+                  "media"
+                ],
+                "type": "string"
+              },
+              "type": {
+                "enum": [
+                  "https://api.twitter.com/2/problems/resource-not-found"
+                ],
+                "type": "string"
+              },
+              "value": {
+                "description": "Value will match the schema of the field."
+              }
+            },
+            "required": [
+              "parameter",
+              "value",
+              "resource_id",
+              "resource_type"
+            ]
+          }
+        ],
+        "description": "A problem that indicates that a given Tweet, User, etc. does not exist."
+      },
+      "ResourceUnauthorizedProblem": {
+        "allOf": [
+          {
+            "$ref": "#/components/schemas/ProblemFields"
+          },
+          {
+            "properties": {
+              "parameter": {
+                "type": "string"
+              },
+              "resource_id": {
+                "type": "string"
+              },
+              "resource_type": {
+                "enum": [
+                  "tweet",
+                  "user",
+                  "media"
+                ],
+                "type": "string"
+              },
+              "section": {
+                "enum": [
+                  "data",
+                  "includes"
+                ],
+                "type": "string"
+              },
+              "type": {
+                "enum": [
+                  "https://api.twitter.com/2/problems/not-authorized-for-resource"
+                ],
+                "type": "string"
+              },
+              "value": {
+                "type": "string"
+              }
+            },
+            "required": [
+              "value",
+              "resource_id",
+              "resource_type",
+              "section",
+              "parameter"
+            ],
+            "type": "object"
+          }
+        ],
+        "description": "A problem that indicates you are not allowed to see a particular Tweet, User, etc.",
+        "type": "object"
+      },
+      "ResourceUnavailableProblem": {
+        "allOf": [
+          {
+            "$ref": "#/components/schemas/ProblemFields"
+          },
+          {
+            "properties": {
+              "parameter": {
+                "minLength": 1,
+                "type": "string"
+              },
+              "resource_id": {
+                "type": "string"
+              },
+              "resource_type": {
+                "enum": [
+                  "user",
+                  "tweet",
+                  "media"
+                ],
+                "type": "string"
+              },
+              "type": {
+                "enum": [
+                  "https://api.twitter.com/2/problems/resource-unavailable"
+                ],
+                "type": "string"
+              }
+            },
+            "required": [
+              "parameter",
+              "resource_id",
+              "resource_type"
+            ]
+          }
+        ],
+        "description": "A problem that indicates a particular Tweet, User, etc. is not available to you."
+      },
+      "Rule": {
+        "description": "A user-provided stream filtering rule.",
+        "properties": {
+          "id": {
+            "$ref": "#/components/schemas/RuleId"
+          },
+          "tag": {
+            "$ref": "#/components/schemas/RuleTag"
+          },
+          "value": {
+            "$ref": "#/components/schemas/RuleValue"
+          }
+        },
+        "required": [
+          "value"
+        ],
+        "type": "object"
+      },
+      "RuleId": {
+        "description": "Unique identifier of this rule.",
+        "example": "120897978112909812",
+        "pattern": "^[0-9]{1,19}$",
+        "type": "string"
+      },
+      "RuleNoId": {
+        "description": "A user-provided stream filtering rule.",
+        "properties": {
+          "tag": {
+            "$ref": "#/components/schemas/RuleTag"
+          },
+          "value": {
+            "$ref": "#/components/schemas/RuleValue"
+          }
+        },
+        "required": [
+          "value"
+        ],
+        "type": "object"
+      },
+      "RuleTag": {
+        "description": "A tag meant for the labeling of user provided rules.",
+        "example": "Non-retweeted coffee tweets",
+        "type": "string"
+      },
+      "RuleValue": {
+        "description": "The filterlang value of the rule.",
+        "example": "coffee -is:retweet",
+        "type": "string"
+      },
+      "RulesCapProblem": {
+        "allOf": [
+          {
+            "$ref": "#/components/schemas/ProblemFields"
+          },
+          {
+            "properties": {
+              "type": {
+                "enum": [
+                  "https://api.twitter.com/2/problems/rule-cap"
+                ],
+                "type": "string"
+              }
+            }
+          }
+        ],
+        "description": "You have exceeded the maximum number of rules.",
+        "type": "object"
+      },
+      "RulesRequestSummary": {
+        "oneOf": [
+          {
+            "description": "A summary of the results of the addition of user-specified stream filtering rules.",
+            "properties": {
+              "created": {
+                "description": "Number of user-specified stream filtering rules that were created.",
+                "example": 1,
+                "format": "int32",
+                "type": "integer"
+              },
+              "invalid": {
+                "description": "Number of invalid user-specified stream filtering rules.",
+                "example": 1,
+                "format": "int32",
+                "type": "integer"
+              },
+              "not_created": {
+                "description": "Number of user-specified stream filtering rules that were not created.",
+                "example": 1,
+                "format": "int32",
+                "type": "integer"
+              },
+              "valid": {
+                "description": "Number of valid user-specified stream filtering rules.",
+                "example": 1,
+                "format": "int32",
+                "type": "integer"
+              }
+            },
+            "required": [
+              "created",
+              "not_created",
+              "valid",
+              "invalid"
+            ],
+            "type": "object"
+          },
+          {
+            "properties": {
+              "deleted": {
+                "description": "Number of user-specified stream filtering rules that were deleted.",
+                "format": "int32",
+                "type": "integer"
+              },
+              "not_deleted": {
+                "description": "Number of user-specified stream filtering rules that were not deleted.",
+                "format": "int32",
+                "type": "integer"
+              }
+            },
+            "required": [
+              "deleted",
+              "not_deleted"
+            ],
+            "type": "object"
+          }
+        ],
+        "type": "object"
+      },
+      "RulesResponseMetadata": {
+        "properties": {
+          "sent": {
+            "type": "string"
+          },
+          "summary": {
+            "$ref": "#/components/schemas/RulesRequestSummary"
+          }
+        },
+        "required": [
+          "sent"
+        ],
+        "type": "object"
+      },
+      "SearchCount": {
+        "description": "Represent a Search Count Result",
+        "properties": {
+          "end": {
+            "$ref": "#/components/schemas/End"
+          },
+          "start": {
+            "$ref": "#/components/schemas/Start"
+          },
+          "tweet_count": {
+            "$ref": "#/components/schemas/TweetCount"
+          }
+        },
+        "required": [
+          "end",
+          "start",
+          "tweet_count"
+        ],
+        "type": "object"
+      },
+      "SingleComplianceJobResponse": {
+        "properties": {
+          "data": {
+            "$ref": "#/components/schemas/ComplianceJob"
+          },
+          "errors": {
+            "items": {
+              "$ref": "#/components/schemas/Problem"
+            },
+            "minItems": 1,
+            "type": "array"
+          }
+        },
+        "type": "object"
+      },
+      "SingleSpaceLookupResponse": {
+        "properties": {
+          "data": {
+            "$ref": "#/components/schemas/Space"
+          },
+          "errors": {
+            "items": {
+              "$ref": "#/components/schemas/Problem"
+            },
+            "minItems": 1,
+            "type": "array"
+          },
+          "includes": {
+            "$ref": "#/components/schemas/Expansions"
+          }
+        },
+        "type": "object"
+      },
+      "SingleTweetLookupResponse": {
+        "properties": {
+          "data": {
+            "$ref": "#/components/schemas/Tweet"
+          },
+          "errors": {
+            "items": {
+              "$ref": "#/components/schemas/Problem"
+            },
+            "minItems": 1,
+            "type": "array"
+          },
+          "includes": {
+            "$ref": "#/components/schemas/Expansions"
+          }
+        },
+        "type": "object"
+      },
+      "SingleUserLookupResponse": {
+        "properties": {
+          "data": {
+            "$ref": "#/components/schemas/User"
+          },
+          "errors": {
+            "items": {
+              "$ref": "#/components/schemas/Problem"
+            },
+            "minItems": 1,
+            "type": "array"
+          },
+          "includes": {
+            "$ref": "#/components/schemas/Expansions"
+          }
+        },
+        "type": "object"
+      },
+      "Space": {
+        "description": "",
+        "properties": {
+          "created_at": {
+            "description": "Creation time of the space",
+            "example": "2021-07-06T18:40:40.000Z",
+            "format": "date-time",
+            "type": "string"
+          },
+          "creator_id": {
+            "$ref": "#/components/schemas/UserID"
+          },
+          "host_ids": {
+            "description": "The user ids for the hosts of the space",
+            "items": {
+              "$ref": "#/components/schemas/UserID"
+            },
+            "type": "array"
+          },
+          "id": {
+            "$ref": "#/components/schemas/SpaceID"
+          },
+          "invited_user_ids": {
+            "description": "An array of user ids for people who were invited to a space",
+            "items": {
+              "$ref": "#/components/schemas/UserID"
+            },
+            "type": "array"
+          },
+          "is_ticketed": {
+            "description": "Denotes if the space is a ticketed space",
+            "example": "false",
+            "type": "boolean"
+          },
+          "lang": {
+            "description": "The language of the space",
+            "example": "en",
+            "type": "string"
+          },
+          "participant_count": {
+            "description": "The number of participants in a space",
+            "example": 10,
+            "format": "int32",
+            "type": "integer"
+          },
+          "scheduled_start": {
+            "description": "A date time stamp for when a space is scheduled to begin",
+            "example": "2021-07-06T18:40:40.000Z",
+            "format": "date-time",
+            "type": "string"
+          },
+          "speaker_ids": {
+            "description": "An array of user ids for people who were speakers in a space",
+            "items": {
+              "$ref": "#/components/schemas/UserID"
+            },
+            "type": "array"
+          },
+          "started_at": {
+            "description": "When the space was started as a date string",
+            "example": "2021-7-14T04:35:55Z",
+            "format": "date-time",
+            "type": "string"
+          },
+          "state": {
+            "description": "The current state of the space.",
+            "enum": [
+              "live",
+              "scheduled",
+              "ended"
+            ],
+            "example": "live",
+            "type": "string"
+          },
+          "title": {
+            "description": "The title of the space",
+            "example": "Spaces are Awesome",
+            "type": "string"
+          },
+          "updated_at": {
+            "description": "When the space was last updated",
+            "example": "2021-7-14T04:35:55Z",
+            "format": "date-time",
+            "type": "string"
+          }
+        },
+        "required": [
+          "id",
+          "state"
+        ],
+        "type": "object"
+      },
+      "SpaceID": {
+        "description": "The unique identifier of this Space.",
+        "example": "1SLjjRYNejbKM",
+        "pattern": "^[a-zA-Z0-9]{1,13}$",
+        "type": "string"
+      },
+      "Start": {
+        "description": "The start time of the bucket",
+        "format": "date-time",
+        "type": "string"
+      },
+      "StreamingTweet": {
+        "description": "A tweet or error that can be returned by the streaming tweet API.",
+        "oneOf": [
+          {
+            "$ref": "#/components/schemas/SingleTweetLookupResponse"
+          },
+          {
+            "properties": {
+              "errors": {
+                "items": {
+                  "$ref": "#/components/schemas/Problem"
+                },
+                "type": "array"
+              }
+            },
+            "required": [
+              "errors"
+            ],
+            "type": "object"
+          }
+        ],
+        "type": "object"
+      },
+      "Tweet": {
+        "example": {
+          "author_id": "2244994945",
+          "created_at": "Wed Jan 06 18:40:40 +0000 2021",
+          "id": "1346889436626259968",
+          "text": "Learn how to use the user Tweet timeline and user mention timeline endpoints in the Twitter API v2 to explore Tweet\\u2026 https:\\/\\/t.co\\/56a0vZUx7i"
+        },
+        "properties": {
+          "attachments": {
+            "description": "Specifies the type of attachments (if any) present in this Tweet.",
+            "properties": {
+              "media_keys": {
+                "description": "A list of Media Keys for each one of the media attachments (if media are attached).",
+                "items": {
+                  "$ref": "#/components/schemas/MediaKey"
+                },
+                "minItems": 1,
+                "type": "array"
+              },
+              "poll_ids": {
+                "description": "A list of poll IDs (if polls are attached).",
+                "items": {
+                  "$ref": "#/components/schemas/PollId"
+                },
+                "minItems": 1,
+                "type": "array"
+              }
+            },
+            "type": "object"
+          },
+          "author_id": {
+            "$ref": "#/components/schemas/UserID"
+          },
+          "context_annotations": {
+            "items": {
+              "$ref": "#/components/schemas/ContextAnnotation"
+            },
+            "minItems": 1,
+            "type": "array"
+          },
+          "conversation_id": {
+            "$ref": "#/components/schemas/TweetID"
+          },
+          "created_at": {
+            "description": "Creation time of the Tweet.",
+            "example": "2021-01-06T18:40:40.000Z",
+            "format": "date-time",
+            "type": "string"
+          },
+          "entities": {
+            "$ref": "#/components/schemas/FullTextEntities"
+          },
+          "geo": {
+            "description": "The location tagged on the Tweet, if the user provided one.",
+            "properties": {
+              "coordinates": {
+                "$ref": "#/components/schemas/Point"
+              },
+              "place_id": {
+                "$ref": "#/components/schemas/PlaceId"
+              }
+            },
+            "type": "object"
+          },
+          "id": {
+            "$ref": "#/components/schemas/TweetID"
+          },
+          "in_reply_to_user_id": {
+            "$ref": "#/components/schemas/UserID"
+          },
+          "lang": {
+            "description": "Language of the Tweet, if detected by Twitter. Returned as a BCP47 language tag.",
+            "example": "en",
+            "type": "string"
+          },
+          "non_public_metrics": {
+            "description": "Nonpublic engagement metrics for the Tweet at the time of the request.",
+            "properties": {
+              "impression_count": {
+                "description": "Number of times this Tweet has been viewed.",
+                "format": "int32",
+                "type": "integer"
+              }
+            },
+            "type": "object"
+          },
+          "organic_metrics": {
+            "description": "Organic nonpublic engagement metrics for the Tweet at the time of the request.",
+            "properties": {
+              "impression_count": {
+                "description": "Number of times this Tweet has been viewed.",
+                "type": "integer"
+              },
+              "like_count": {
+                "description": "Number of times this Tweet has been liked.",
+                "type": "integer"
+              },
+              "reply_count": {
+                "description": "Number of times this Tweet has been replied to.",
+                "type": "integer"
+              },
+              "retweet_count": {
+                "description": "Number of times this Tweet has been Retweeted.",
+                "type": "integer"
+              }
+            },
+            "required": [
+              "impression_count",
+              "retweet_count",
+              "reply_count",
+              "like_count"
+            ],
+            "type": "object"
+          },
+          "possibly_sensitive": {
+            "description": "Indicates if this Tweet contains URLs marked as sensitive, for example content suitable for mature audiences.",
+            "example": false,
+            "type": "boolean"
+          },
+          "promoted_metrics": {
+            "description": "Promoted nonpublic engagement metrics for the Tweet at the time of the request.",
+            "properties": {
+              "impression_count": {
+                "description": "Number of times this Tweet has been viewed.",
+                "format": "int32",
+                "type": "integer"
+              },
+              "like_count": {
+                "description": "Number of times this Tweet has been liked.",
+                "format": "int32",
+                "type": "integer"
+              },
+              "reply_count": {
+                "description": "Number of times this Tweet has been replied to.",
+                "format": "int32",
+                "type": "integer"
+              },
+              "retweet_count": {
+                "description": "Number of times this Tweet has been Retweeted.",
+                "format": "int32",
+                "type": "integer"
+              }
+            },
+            "type": "object"
+          },
+          "public_metrics": {
+            "description": "Engagement metrics for the Tweet at the time of the request.",
+            "properties": {
+              "like_count": {
+                "description": "Number of times this Tweet has been liked.",
+                "type": "integer"
+              },
+              "quote_count": {
+                "description": "Number of times this Tweet has been quoted.",
+                "type": "integer"
+              },
+              "reply_count": {
+                "description": "Number of times this Tweet has been replied to.",
+                "type": "integer"
+              },
+              "retweet_count": {
+                "description": "Number of times this Tweet has been Retweeted.",
+                "type": "integer"
+              }
+            },
+            "required": [
+              "retweet_count",
+              "reply_count",
+              "like_count"
+            ],
+            "type": "object"
+          },
+          "referenced_tweets": {
+            "description": "A list of Tweets this Tweet refers to. For example, if the parent Tweet is a Retweet, a Quoted Tweet or a Reply, it will include the related Tweet referenced to by its parent.",
+            "items": {
+              "properties": {
+                "id": {
+                  "$ref": "#/components/schemas/TweetID"
+                },
+                "type": {
+                  "enum": [
+                    "retweeted",
+                    "quoted",
+                    "replied_to"
+                  ],
+                  "type": "string"
+                }
+              },
+              "required": [
+                "type",
+                "id"
+              ],
+              "type": "object"
+            },
+            "minItems": 1,
+            "type": "array"
+          },
+          "reply_settings": {
+            "$ref": "#/components/schemas/ReplySettings"
+          },
+          "source": {
+            "description": "The name of the app the user Tweeted from.",
+            "type": "string"
+          },
+          "text": {
+            "description": "The content of the Tweet.",
+            "example": "Learn how to use the user Tweet timeline and user mention timeline endpoints in the Twitter API v2 to explore Tweet\\u2026 https:\\/\\/t.co\\/56a0vZUx7i",
+            "type": "string"
+          },
+          "withheld": {
+            "$ref": "#/components/schemas/TweetWithheld"
+          }
+        },
+        "required": [
+          "id",
+          "text"
+        ],
+        "type": "object"
+      },
+      "TweetCount": {
+        "description": "The count for the bucket",
+        "type": "integer"
+      },
+      "TweetCountsResponse": {
+        "properties": {
+          "data": {
+            "items": {
+              "$ref": "#/components/schemas/SearchCount"
+            },
+            "minItems": 1,
+            "type": "array"
+          },
+          "errors": {
+            "items": {
+              "$ref": "#/components/schemas/Problem"
+            },
+            "minItems": 1,
+            "type": "array"
+          },
+          "meta": {
+            "properties": {
+              "next_token": {
+                "description": "This value is used to get the next 'page' of results by providing it to the next_token parameter.",
+                "type": "string"
+              },
+              "total_tweet_count": {
+                "description": "Sum of search query count results",
+                "format": "int32",
+                "type": "integer"
+              }
+            },
+            "type": "object"
+          }
+        },
+        "type": "object"
+      },
+      "TweetDeleteResponse": {
+        "properties": {
+          "data": {
+            "properties": {
+              "deleted": {
+                "type": "boolean"
+              }
+            },
+            "required": [
+              "deleted"
+            ],
+            "type": "object"
+          },
+          "errors": {
+            "items": {
+              "$ref": "#/components/schemas/Problem"
+            },
+            "minItems": 1,
+            "type": "array"
+          }
+        },
+        "type": "object"
+      },
+      "TweetID": {
+        "description": "Unique identifier of this Tweet. This is returned as a string in order to avoid complications with languages and tools that cannot handle large integers.",
+        "example": "1346889436626259968",
+        "pattern": "^[0-9]{1,19}$",
+        "type": "string"
+      },
+      "TweetSearchResponse": {
+        "properties": {
+          "data": {
+            "items": {
+              "$ref": "#/components/schemas/Tweet"
+            },
+            "minItems": 1,
+            "type": "array"
+          },
+          "errors": {
+            "items": {
+              "$ref": "#/components/schemas/Problem"
+            },
+            "minItems": 1,
+            "type": "array"
+          },
+          "includes": {
+            "$ref": "#/components/schemas/Expansions"
+          },
+          "meta": {
+            "properties": {
+              "newest_id": {
+                "description": "Most recent Tweet Id returned by search query",
+                "pattern": "^[0-9]{1,19}$",
+                "type": "string"
+              },
+              "next_token": {
+                "description": "This value is used to get the next 'page' of results by providing it to the next_token parameter.",
+                "type": "string"
+              },
+              "oldest_id": {
+                "description": "Oldest Tweet Id returned by search query",
+                "pattern": "^[0-9]{1,19}$",
+                "type": "string"
+              },
+              "result_count": {
+                "description": "Number of search query results",
+                "format": "int32",
+                "type": "integer"
+              }
+            },
+            "type": "object"
+          }
+        },
+        "type": "object"
+      },
+      "TweetWithheld": {
+        "description": "Indicates withholding details for [withheld content](https://help.twitter.com/en/rules-and-policies/tweet-withheld-by-country).",
+        "properties": {
+          "copyright": {
+            "description": "Indicates if the content is being withheld for on the basis of copyright infringement.",
+            "type": "boolean"
+          },
+          "country_codes": {
+            "description": "Provides a list of countries where this content is not available.",
+            "items": {
+              "$ref": "#/components/schemas/CountryCode"
+            },
+            "minItems": 1,
+            "type": "array",
+            "uniqueItems": true
+          },
+          "scope": {
+            "description": "Indicates whether the content being withheld is the `tweet` or a `user`.",
+            "enum": [
+              "tweet",
+              "user"
+            ],
+            "type": "string"
+          }
+        },
+        "required": [
+          "copyright",
+          "country_codes"
+        ],
+        "type": "object"
+      },
+      "URL": {
+        "description": "A validly formatted URL.",
+        "example": "https://developer.twitter.com/en/docs/twitter-api",
+        "format": "uri",
+        "type": "string"
+      },
+      "URLFields": {
+        "description": "Represent the portion of text recognized as a URL.",
+        "properties": {
+          "description": {
+            "description": "Description of the URL landing page.",
+            "example": "This is a description of the website.",
+            "type": "string"
+          },
+          "display_url": {
+            "description": "The URL as displayed in the Twitter client.",
+            "example": "twittercommunity.com/t/introducing-…",
+            "type": "string"
+          },
+          "expanded_url": {
+            "$ref": "#/components/schemas/URL"
+          },
+          "images": {
+            "items": {
+              "$ref": "#/components/schemas/URLImage"
+            },
+            "minItems": 1,
+            "type": "array"
+          },
+          "status": {
+            "$ref": "#/components/schemas/HTTPStatusCode"
+          },
+          "title": {
+            "description": "Title of the page the URL points to.",
+            "example": "Introducing the v2 follow lookup endpoints",
+            "type": "string"
+          },
+          "unwound_url": {
+            "description": "Fully resolved url",
+            "example": "https://twittercommunity.com/t/introducing-the-v2-follow-lookup-endpoints/147118",
+            "format": "uri",
+            "type": "string"
+          },
+          "url": {
+            "$ref": "#/components/schemas/URL"
+          }
+        },
+        "required": [
+          "url"
+        ],
+        "type": "object"
+      },
+      "URLImage": {
+        "description": "Represent the information for the URL image",
+        "properties": {
+          "height": {
+            "$ref": "#/components/schemas/MediaHeight"
+          },
+          "url": {
+            "$ref": "#/components/schemas/URL"
+          },
+          "width": {
+            "$ref": "#/components/schemas/MediaWidth"
+          }
+        },
+        "type": "object"
+      },
+      "UnsupportedAuthenticationProblem": {
+        "allOf": [
+          {
+            "$ref": "#/components/schemas/ProblemFields"
+          },
+          {
+            "properties": {
+              "type": {
+                "enum": [
+                  "https://api.twitter.com/2/problems/unsupported-authentication"
+                ],
+                "type": "string"
+              }
+            },
+            "type": "object"
+          }
+        ],
+        "description": "A problem that indicates that the authentication used is not supported.",
+        "type": "object"
+      },
+      "UploadExpiration": {
+        "description": "Expiration time of the upload URL",
+        "example": "2021-01-06T18:40:40.000Z",
+        "format": "date-time",
+        "type": "string"
+      },
+      "UploadUrl": {
+        "description": "URL to which the user will upload their tweet or user IDs",
+        "format": "uri",
+        "type": "string"
+      },
+      "UrlEntity": {
+        "allOf": [
+          {
+            "$ref": "#/components/schemas/EntityIndicesInclusiveExclusive"
+          },
+          {
+            "$ref": "#/components/schemas/URLFields"
+          }
+        ],
+        "description": "Represent the portion of text recognized as a URL, and its start and end position within the text."
+      },
+      "UsageCapExceededProblem": {
+        "allOf": [
+          {
+            "$ref": "#/components/schemas/ProblemFields"
+          },
+          {
+            "properties": {
+              "period": {
+                "enum": [
+                  "Daily",
+                  "Monthly"
+                ],
+                "type": "string"
+              },
+              "scope": {
+                "enum": [
+                  "Account",
+                  "Product"
+                ],
+                "type": "string"
+              },
+              "type": {
+                "enum": [
+                  "https://api.twitter.com/2/problems/usage-capped"
+                ],
+                "type": "string"
+              }
+            },
+            "type": "object"
+          }
+        ],
+        "description": "A problem that indicates that a usage cap has been exceeded.",
+        "type": "object"
+      },
+      "User": {
+        "description": "The Twitter User object",
+        "example": {
+          "created_at": "2013-12-14T04:35:55Z",
+          "id": "2244994945",
+          "name": "Twitter Dev",
+          "protected": false,
+          "username": "TwitterDev"
+        },
+        "properties": {
+          "created_at": {
+            "description": "Creation time of this user.",
+            "format": "date-time",
+            "type": "string"
+          },
+          "description": {
+            "description": "The text of this user's profile description (also known as bio), if the user provided one.",
+            "type": "string"
+          },
+          "entities": {
+            "description": "A list of metadata found in the user's profile description.",
+            "properties": {
+              "description": {
+                "$ref": "#/components/schemas/FullTextEntities"
+              },
+              "url": {
+                "description": "Expanded details for the URL specified in the user's profile, with start and end indices.",
+                "properties": {
+                  "urls": {
+                    "items": {
+                      "$ref": "#/components/schemas/UrlEntity"
+                    },
+                    "minItems": 1,
+                    "type": "array"
+                  }
+                },
+                "type": "object"
+              }
+            },
+            "type": "object"
+          },
+          "id": {
+            "$ref": "#/components/schemas/UserID"
+          },
+          "location": {
+            "description": "The location specified in the user's profile, if the user provided one. As this is a freeform value, it may not indicate a valid location, but it may be fuzzily evaluated when performing searches with location queries.",
+            "type": "string"
+          },
+          "name": {
+            "description": "The friendly name of this user, as shown on their profile.",
+            "type": "string"
+          },
+          "pinned_tweet_id": {
+            "$ref": "#/components/schemas/TweetID"
+          },
+          "profile_image_url": {
+            "description": "The URL to the profile image for this user.",
+            "format": "uri",
+            "type": "string"
+          },
+          "protected": {
+            "description": "Indicates if this user has chosen to protect their Tweets (in other words, if this user's Tweets are private).",
+            "type": "boolean"
+          },
+          "public_metrics": {
+            "description": "A list of metrics for this user",
+            "properties": {
+              "followers_count": {
+                "description": "Number of users who are following this user.",
+                "type": "integer"
+              },
+              "following_count": {
+                "description": "Number of users this user is following.",
+                "type": "integer"
+              },
+              "listed_count": {
+                "description": "The number of lists that include this user.",
+                "type": "integer"
+              },
+              "tweet_count": {
+                "description": "The number of Tweets (including Retweets) posted by this user.",
+                "type": "integer"
+              }
+            },
+            "required": [
+              "followers_count",
+              "following_count",
+              "tweet_count",
+              "listed_count"
+            ],
+            "type": "object"
+          },
+          "url": {
+            "description": "The URL specified in the user's profile.",
+            "type": "string"
+          },
+          "username": {
+            "$ref": "#/components/schemas/UserName"
+          },
+          "verified": {
+            "description": "Indicate if this user is a verified Twitter User.",
+            "type": "boolean"
+          },
+          "withheld": {
+            "$ref": "#/components/schemas/UserWithheld"
+          }
+        },
+        "required": [
+          "id",
+          "name",
+          "username"
+        ]
+      },
+      "UserID": {
+        "description": "Unique identifier of this User. This is returned as a string in order to avoid complications with languages and tools that cannot handle large integers.",
+        "example": "2244994945",
+        "pattern": "^[0-9]{1,19}$",
+        "type": "string"
+      },
+      "UserName": {
+        "description": "The Twitter handle (screen name) of this user.",
+        "pattern": "^[A-Za-z0-9_]{1,15}$",
+        "type": "string"
+      },
+      "UserWithheld": {
+        "description": "Indicates withholding details for [withheld content](https://help.twitter.com/en/rules-and-policies/tweet-withheld-by-country).",
+        "properties": {
+          "country_codes": {
+            "description": "Provides a list of countries where this content is not available.",
+            "items": {
+              "$ref": "#/components/schemas/CountryCode"
+            },
+            "minItems": 1,
+            "type": "array",
+            "uniqueItems": true
+          },
+          "scope": {
+            "description": "Indicates that the content being withheld is a `user`.",
+            "enum": [
+              "user"
+            ],
+            "type": "string"
+          }
+        },
+        "required": [
+          "country_codes"
+        ],
+        "type": "object"
+      },
+      "UsersBlockingMutationResponse": {
+        "properties": {
+          "data": {
+            "properties": {
+              "blocking": {
+                "type": "boolean"
+              }
+            },
+            "type": "object"
+          },
+          "errors": {
+            "items": {
+              "$ref": "#/components/schemas/Problem"
+            },
+            "minItems": 1,
+            "type": "array"
+          }
+        },
+        "type": "object"
+      },
+      "UsersFollowingCreateResponse": {
+        "properties": {
+          "data": {
+            "properties": {
+              "following": {
+                "type": "boolean"
+              },
+              "pending_follow": {
+                "type": "boolean"
+              }
+            },
+            "type": "object"
+          },
+          "errors": {
+            "items": {
+              "$ref": "#/components/schemas/Problem"
+            },
+            "minItems": 1,
+            "type": "array"
+          }
+        },
+        "type": "object"
+      },
+      "UsersFollowingDeleteResponse": {
+        "properties": {
+          "data": {
+            "properties": {
+              "following": {
+                "type": "boolean"
+              }
+            },
+            "type": "object"
+          },
+          "errors": {
+            "items": {
+              "$ref": "#/components/schemas/Problem"
+            },
+            "minItems": 1,
+            "type": "array"
+          }
+        },
+        "type": "object"
+      },
+      "UsersFollowingLookupResponse": {
+        "properties": {
+          "data": {
+            "items": {
+              "$ref": "#/components/schemas/User"
+            },
+            "type": "array"
+          },
+          "errors": {
+            "items": {
+              "$ref": "#/components/schemas/Problem"
+            },
+            "minItems": 1,
+            "type": "array"
+          },
+          "includes": {
+            "$ref": "#/components/schemas/Expansions"
+          },
+          "meta": {
+            "properties": {
+              "next_token": {
+                "description": "The next token",
+                "pattern": "^[a-zA-Z0-9]+$",
+                "type": "string"
+              },
+              "previous_token": {
+                "description": "The previous token",
+                "pattern": "^[a-zA-Z0-9]+$",
+                "type": "string"
+              },
+              "result_count": {
+                "description": "The number of user results returned in this response",
+                "type": "integer"
+              }
+            },
+            "type": "object"
+          }
+        },
+        "type": "object"
+      },
+      "UsersLikesCreateRequest": {
+        "properties": {
+          "tweet_id": {
+            "pattern": "^[0-9]{1,19}$",
+            "type": "string"
+          }
+        },
+        "required": [
+          "tweet_id"
+        ],
+        "type": "object"
+      },
+      "UsersLikesCreateResponse": {
+        "properties": {
+          "data": {
+            "properties": {
+              "liked": {
+                "type": "boolean"
+              }
+            },
+            "type": "object"
+          },
+          "errors": {
+            "items": {
+              "$ref": "#/components/schemas/Problem"
+            },
+            "minItems": 1,
+            "type": "array"
+          }
+        },
+        "type": "object"
+      },
+      "UsersLikesDeleteResponse": {
+        "properties": {
+          "data": {
+            "properties": {
+              "liked": {
+                "type": "boolean"
+              }
+            },
+            "type": "object"
+          },
+          "errors": {
+            "items": {
+              "$ref": "#/components/schemas/Problem"
+            },
+            "minItems": 1,
+            "type": "array"
+          }
+        },
+        "type": "object"
+      },
+      "UsersMutingMutationResponse": {
+        "properties": {
+          "data": {
+            "properties": {
+              "muting": {
+                "type": "boolean"
+              }
+            },
+            "type": "object"
+          },
+          "errors": {
+            "items": {
+              "$ref": "#/components/schemas/Problem"
+            },
+            "minItems": 1,
+            "type": "array"
+          }
+        },
+        "type": "object"
+      },
+      "UsersRetweetsCreateRequest": {
+        "properties": {
+          "tweet_id": {
+            "pattern": "^[0-9]{1,19}$",
+            "type": "string"
+          }
+        },
+        "required": [
+          "tweet_id"
+        ],
+        "type": "object"
+      },
+      "UsersRetweetsCreateResponse": {
+        "properties": {
+          "data": {
+            "properties": {
+              "retweeted": {
+                "type": "boolean"
+              }
+            },
+            "type": "object"
+          },
+          "errors": {
+            "items": {
+              "$ref": "#/components/schemas/Problem"
+            },
+            "minItems": 1,
+            "type": "array"
+          }
+        },
+        "type": "object"
+      },
+      "UsersRetweetsDeleteResponse": {
+        "properties": {
+          "data": {
+            "properties": {
+              "retweeted": {
+                "type": "boolean"
+              }
+            },
+            "type": "object"
+          },
+          "errors": {
+            "items": {
+              "$ref": "#/components/schemas/Problem"
+            },
+            "minItems": 1,
+            "type": "array"
+          }
+        },
+        "type": "object"
+      },
+      "Video": {
+        "allOf": [
+          {
+            "$ref": "#/components/schemas/CommonMediaFields"
+          },
+          {
+            "properties": {
+              "duration_ms": {
+                "type": "integer"
+              },
+              "non_public_metrics": {
+                "description": "Nonpublic engagement metrics for the Media at the time of the request.",
+                "properties": {
+                  "playback_0_count": {
+                    "description": "Number of users who made it through 0% of the video.",
+                    "format": "int32",
+                    "type": "integer"
+                  },
+                  "playback_100_count": {
+                    "description": "Number of users who made it through 100% of the video.",
+                    "format": "int32",
+                    "type": "integer"
+                  },
+                  "playback_25_count": {
+                    "description": "Number of users who made it through 25% of the video.",
+                    "format": "int32",
+                    "type": "integer"
+                  },
+                  "playback_50_count": {
+                    "description": "Number of users who made it through 50% of the video.",
+                    "format": "int32",
+                    "type": "integer"
+                  },
+                  "playback_75_count": {
+                    "description": "Number of users who made it through 75% of the video.",
+                    "format": "int32",
+                    "type": "integer"
+                  }
+                },
+                "type": "object"
+              },
+              "organic_metrics": {
+                "description": "Organic nonpublic engagement metrics for the Media at the time of the request.",
+                "properties": {
+                  "playback_0_count": {
+                    "description": "Number of users who made it through 0% of the video.",
+                    "format": "int32",
+                    "type": "integer"
+                  },
+                  "playback_100_count": {
+                    "description": "Number of users who made it through 100% of the video.",
+                    "format": "int32",
+                    "type": "integer"
+                  },
+                  "playback_25_count": {
+                    "description": "Number of users who made it through 25% of the video.",
+                    "format": "int32",
+                    "type": "integer"
+                  },
+                  "playback_50_count": {
+                    "description": "Number of users who made it through 50% of the video.",
+                    "format": "int32",
+                    "type": "integer"
+                  },
+                  "playback_75_count": {
+                    "description": "Number of users who made it through 75% of the video.",
+                    "format": "int32",
+                    "type": "integer"
+                  },
+                  "view_count": {
+                    "description": "Number of times this video has been viewed.",
+                    "format": "int32",
+                    "type": "integer"
+                  }
+                },
+                "type": "object"
+              },
+              "preview_image_url": {
+                "format": "uri",
+                "type": "string"
+              },
+              "promoted_metrics": {
+                "description": "Promoted nonpublic engagement metrics for the Media at the time of the request.",
+                "properties": {
+                  "playback_0_count": {
+                    "description": "Number of users who made it through 0% of the video.",
+                    "format": "int32",
+                    "type": "integer"
+                  },
+                  "playback_100_count": {
+                    "description": "Number of users who made it through 100% of the video.",
+                    "format": "int32",
+                    "type": "integer"
+                  },
+                  "playback_25_count": {
+                    "description": "Number of users who made it through 25% of the video.",
+                    "format": "int32",
+                    "type": "integer"
+                  },
+                  "playback_50_count": {
+                    "description": "Number of users who made it through 50% of the video.",
+                    "format": "int32",
+                    "type": "integer"
+                  },
+                  "playback_75_count": {
+                    "description": "Number of users who made it through 75% of the video.",
+                    "format": "int32",
+                    "type": "integer"
+                  },
+                  "view_count": {
+                    "description": "Number of times this video has been viewed.",
+                    "format": "int32",
+                    "type": "integer"
+                  }
+                },
+                "type": "object"
+              },
+              "public_metrics": {
+                "description": "Engagement metrics for the Media at the time of the request.",
+                "properties": {
+                  "view_count": {
+                    "description": "Number of times this video has been viewed.",
+                    "format": "int32",
+                    "type": "integer"
+                  }
+                },
+                "type": "object"
+              },
+              "type": {
+                "enum": [
+                  "video"
+                ],
+                "type": "string"
+              }
+            },
+            "type": "object"
+          }
+        ],
+        "type": "object"
+      }
+    },
+    "securitySchemes": {
+      "BearerToken": {
+        "scheme": "bearer",
+        "type": "http"
+      },
+      "UserToken": {
+        "scheme": "OAuth",
+        "type": "http"
+      }
+    }
+  }
+}
\ No newline at end of file
