WebAPI missing files pt2

diff --git a/Engine/Plugins/Web/WebAPI/Source/WebAPIOpenAPI/Private/Tests/Samples/V3/uv.json b/Engine/Plugins/Web/WebAPI/Source/WebAPIOpenAPI/Private/Tests/Samples/V3/uv.json
new file mode 100644
--- /dev/null
+++ b/Engine/Plugins/Web/WebAPI/Source/WebAPIOpenAPI/Private/Tests/Samples/V3/uv.json
@@ -0,0 +1,407 @@
+{
+  "openapi": "3.0.0",
+  "servers": [
+    {
+      "url": "https://api.openuv.io/api/{version}",
+      "variables": {
+        "version": {
+          "default": "v1"
+        }
+      }
+    }
+  ],
+  "info": {
+    "description": "The missing minimalistic JSON real-time UV Index API for awesome Developers, Innovators and Smart Home Enthusiasts",
+    "title": "OpenUV - Global Real-Time UV Index Forecast API",
+    "version": "v1",
+    "x-origin": [
+      {
+        "format": "openapi",
+        "url": "https://gist.githubusercontent.com/MikeRalphson/77bb693d4cf9213909527b9cc4566609/raw/9febc5d16d8d162c93f3503510ce0808a8075d5b/openuv.yaml",
+        "version": "3.0"
+      }
+    ],
+    "x-providerName": "openuv.io",
+    "x-logo": {
+      "url": "https://api.apis.guru/v2/cache/logo/https_apis.guru_assets_images_no-logo.svg"
+    }
+  },
+  "paths": {
+    "/forecast": {
+      "description": "Hourly UV Index Forecast",
+      "get": {
+        "description": "Get hourly UV Index Forecast by location and date. Optional altitude, ozone level and datetime could be provided.",
+        "parameters": [
+          {
+            "description": "latitude, from -90.00 to 90.00",
+            "in": "query",
+            "name": "lat",
+            "required": true,
+            "schema": {
+              "description": "latitude, from -90.00 to 90.00",
+              "example": 78.67,
+              "type": "number"
+            }
+          },
+          {
+            "description": "longitude, from -180.00 to 180.00",
+            "in": "query",
+            "name": "lng",
+            "required": true,
+            "schema": {
+              "description": "longitude, from -180.00 to 180.00",
+              "example": 115.67,
+              "type": "number"
+            }
+          },
+          {
+            "description": "Altitude in meters, from 0 to 10000m, 0m by default. If provided the altitude correction factor will be applied to clear sky sea level UV Index value.",
+            "in": "query",
+            "name": "alt",
+            "required": false,
+            "schema": {
+              "description": "Altitude in meters, from 0 to 10000m, 0m by default. If provided the altitude correction factor will be applied to clear sky sea level UV Index value.",
+              "example": 1050,
+              "type": "number"
+            }
+          },
+          {
+            "description": "Ozone in du (Dobson Units), from 100 to 550du, the latest forecast from OMI dataset is used by default.",
+            "in": "query",
+            "name": "ozone",
+            "required": false,
+            "schema": {
+              "description": "Ozone in du (Dobson Units), from 100 to 550du, the latest forecast from OMI dataset is used by default.",
+              "example": 304.5,
+              "type": "number"
+            }
+          },
+          {
+            "description": "UTC datetime in ISO-8601 format, now by default. Use that parameter to get UV Index Forecast for any point in time.",
+            "in": "query",
+            "name": "dt",
+            "required": false,
+            "schema": {
+              "description": "UTC datetime in ISO-8601 format, now by default. Use that parameter to get UV Index Forecast for any point in time.",
+              "example": "2018-02-04T04:39:06.467Z",
+              "format": "date-time",
+              "type": "string"
+            }
+          },
+          {
+            "description": "This header is used to send data that contains your OpenUV API key",
+            "in": "header",
+            "name": "x-access-token",
+            "required": true,
+            "schema": {
+              "description": "This header is used to send data that contains your OpenUV API key",
+              "title": "x-access-token",
+              "type": "string"
+            }
+          }
+        ],
+        "responses": {
+          "200": {
+            "content": {
+              "application/json": {
+                "schema": {
+                  "items": {
+                    "$ref": "#/components/schemas/forecastResult"
+                  },
+                  "type": "array"
+                }
+              }
+            },
+            "description": ""
+          }
+        }
+      }
+    },
+    "/protection": {
+      "description": "Daily Protection Time",
+      "get": {
+        "description": "Get daily protection time by location, UV Index from and UV Index to with 10 minutes accuracy. Optional altitide and ozone level could be provided.",
+        "parameters": [
+          {
+            "description": "latitude, from -90.00 to 90.00",
+            "in": "query",
+            "name": "lat",
+            "required": true,
+            "schema": {
+              "description": "latitude, from -90.00 to 90.00",
+              "example": 78.67,
+              "type": "number"
+            }
+          },
+          {
+            "description": "longitude, from -180.00 to 180.00",
+            "in": "query",
+            "name": "lng",
+            "required": true,
+            "schema": {
+              "description": "longitude, from -180.00 to 180.00",
+              "example": 115.67,
+              "type": "number"
+            }
+          },
+          {
+            "description": "UV Index from value for protection datetime lookup. From 0 to 40.",
+            "in": "query",
+            "name": "from",
+            "required": true,
+            "schema": {
+              "description": "UV Index from value for protection datetime lookup. From 0 to 40.",
+              "example": 2.5,
+              "type": "number"
+            }
+          },
+          {
+            "description": "UV Index to value for protection datetime lookup. From 0 to 40.",
+            "in": "query",
+            "name": "to",
+            "required": true,
+            "schema": {
+              "description": "UV Index to value for protection datetime lookup. From 0 to 40.",
+              "example": 3.6,
+              "type": "number"
+            }
+          },
+          {
+            "description": "Altitude in meters, from 0 to 10000m, 0m by default. If provided the altitude correction factor will be applied to clear sky sea level UV Index value.",
+            "in": "query",
+            "name": "alt",
+            "required": false,
+            "schema": {
+              "description": "Altitude in meters, from 0 to 10000m, 0m by default. If provided the altitude correction factor will be applied to clear sky sea level UV Index value.",
+              "example": 1050,
+              "type": "number"
+            }
+          },
+          {
+            "description": "Ozone in du (Dobson Units), from 100 to 550du, the latest forecast from OMI dataset is used by default.",
+            "in": "query",
+            "name": "ozone",
+            "required": false,
+            "schema": {
+              "description": "Ozone in du (Dobson Units), from 100 to 550du, the latest forecast from OMI dataset is used by default.",
+              "example": 304.5,
+              "type": "number"
+            }
+          },
+          {
+            "description": "This header is used to send data that contains your OpenUV API key",
+            "in": "header",
+            "name": "x-access-token",
+            "required": true,
+            "schema": {
+              "description": "This header is used to send data that contains your OpenUV API key",
+              "title": "x-access-token",
+              "type": "string"
+            }
+          }
+        ],
+        "responses": {
+          "200": {
+            "content": {
+              "application/json": {
+                "schema": {
+                  "$ref": "#/components/schemas/protectionResult",
+                  "type": "object"
+                }
+              }
+            },
+            "description": ""
+          }
+        }
+      }
+    },
+    "/uv": {
+      "description": "Real-time UV Index",
+      "get": {
+        "description": "Get real-time UV Index by location. Optional altitude, ozone level and datetime could be provided.",
+        "parameters": [
+          {
+            "description": "latitude, from -90.00 to 90.00",
+            "in": "query",
+            "name": "lat",
+            "required": true,
+            "schema": {
+              "description": "latitude, from -90.00 to 90.00",
+              "example": 78.67,
+              "type": "number"
+            }
+          },
+          {
+            "description": "longitude, from -180.00 to 180.00",
+            "in": "query",
+            "name": "lng",
+            "required": true,
+            "schema": {
+              "description": "longitude, from -180.00 to 180.00",
+              "example": 115.67,
+              "type": "number"
+            }
+          },
+          {
+            "description": "Altitude in meters, from 0 to 10000m, 0m by default. If provided the altitude correction factor will be applied to clear sky sea level UV Index value.",
+            "in": "query",
+            "name": "alt",
+            "required": false,
+            "schema": {
+              "description": "Altitude in meters, from 0 to 10000m, 0m by default. If provided the altitude correction factor will be applied to clear sky sea level UV Index value.",
+              "example": 1050,
+              "type": "number"
+            }
+          },
+          {
+            "description": "Ozone in du (Dobson Units), from 100 to 550du, the latest forecast from OMI dataset is used by default.",
+            "in": "query",
+            "name": "ozone",
+            "required": false,
+            "schema": {
+              "description": "Ozone in du (Dobson Units), from 100 to 550du, the latest forecast from OMI dataset is used by default.",
+              "example": 304.5,
+              "type": "number"
+            }
+          },
+          {
+            "description": "UTC datetime in ISO-8601 format, now by default. Use that parameter to get UV Index Forecast for any point in time.",
+            "in": "query",
+            "name": "dt",
+            "required": false,
+            "schema": {
+              "description": "UTC datetime in ISO-8601 format, now by default. Use that parameter to get UV Index Forecast for any point in time.",
+              "example": "2018-02-04T04:39:06.467Z",
+              "format": "date-time",
+              "type": "string"
+            }
+          },
+          {
+            "description": "This header is used to send data that contains your OpenUV API key",
+            "in": "header",
+            "name": "x-access-token",
+            "required": true,
+            "schema": {
+              "description": "This header is used to send data that contains your OpenUV API key",
+              "title": "x-access-token",
+              "type": "string"
+            }
+          }
+        ],
+        "responses": {
+          "200": {
+            "content": {
+              "application/json": {
+                "schema": {
+                  "$ref": "#/components/schemas/uvIndexResult",
+                  "type": "object"
+                }
+              }
+            },
+            "description": ""
+          }
+        }
+      }
+    }
+  },
+  "components": {
+    "schemas": {
+      "Error": {
+        "properties": {
+          "error": {
+            "type": "string"
+          }
+        },
+        "required": [
+          "error"
+        ],
+        "type": "object"
+      },
+      "forecast": {
+        "properties": {
+          "uv": {
+            "type": "number"
+          },
+          "uv_time": {
+            "format": "date-time-only",
+            "type": "string"
+          }
+        },
+        "required": [
+          "uv",
+          "uv_time"
+        ],
+        "type": "object"
+      },
+      "forecastResult": {
+        "items": {
+          "$ref": "#/components/schemas/forecast"
+        },
+        "title": "result",
+        "type": "array",
+        "uniqueItems": true
+      },
+      "protectionResult": {
+        "properties": {
+          "ozone": {
+            "type": "number"
+          },
+          "ozone_time": {
+            "format": "date-time-only",
+            "type": "string"
+          },
+          "uv": {
+            "type": "number"
+          },
+          "uv_max": {
+            "type": "number"
+          },
+          "uv_max_time": {
+            "format": "date-time-only",
+            "type": "string"
+          },
+          "uv_time": {
+            "format": "date-time-only",
+            "type": "string"
+          }
+        },
+        "required": [
+          "uv",
+          "uv_time",
+          "uv_max",
+          "uv_max_time",
+          "ozone",
+          "ozone_time"
+        ],
+        "title": "result",
+        "type": "object"
+      },
+      "uvIndexResult": {
+        "properties": {
+          "from_time": {
+            "format": "date-time-only",
+            "type": "string"
+          },
+          "from_uv": {
+            "type": "number"
+          },
+          "to_time": {
+            "format": "date-time-only",
+            "type": "string"
+          },
+          "to_uv": {
+            "type": "number"
+          }
+        },
+        "required": [
+          "from_time",
+          "from_uv",
+          "to_time",
+          "to_uv"
+        ],
+        "title": "result",
+        "type": "object"
+      }
+    }
+  }
+}
\ No newline at end of file
diff --git a/Engine/Plugins/Web/WebAPI/Source/WebAPIOpenAPI/Private/Tests/Samples/digitaltwins_data.json b/Engine/Plugins/Web/WebAPI/Source/WebAPIOpenAPI/Private/Tests/Samples/digitaltwins_data.json
new file mode 100644
--- /dev/null
+++ b/Engine/Plugins/Web/WebAPI/Source/WebAPIOpenAPI/Private/Tests/Samples/digitaltwins_data.json
@@ -0,0 +1,1624 @@
+{
+  "swagger": "2.0",
+  "info": {
+    "version": "2020-10-31",
+    "title": "Azure Digital Twins API",
+    "description": "A service for managing and querying digital twins and digital twin models."
+  },
+  "host": "digitaltwins-hostname",
+  "schemes": [
+    "https"
+  ],
+  "produces": [
+    "application/json"
+  ],
+  "consumes": [
+    "application/json"
+  ],
+  "paths": {
+    "/models": {
+      "parameters": [
+        {
+          "$ref": "#/parameters/traceparent"
+        },
+        {
+          "$ref": "#/parameters/tracestate"
+        }
+      ],
+      "post": {
+        "tags": [
+          "Models"
+        ],
+        "description": "Uploads one or more models. When any error occurs, no models are uploaded.\nStatus codes:\n* 201 Created\n* 400 Bad Request\n  * DTDLParserError - The models provided are not valid DTDL.\n  * InvalidArgument - The model id is invalid.\n  * LimitExceeded - The maximum number of model ids allowed in 'dependenciesFor' has been reached.\n  * ModelVersionNotSupported - The version of DTDL used is not supported.\n* 409 Conflict\n  * ModelAlreadyExists - The model provided already exists.",
+        "operationId": "DigitalTwinModels_Add",
+        "x-ms-examples": {
+          "Add models": {
+            "$ref": "./examples/ModelAdd.json"
+          }
+        },
+        "parameters": [
+          {
+            "name": "models",
+            "in": "body",
+            "description": "An array of models to add.",
+            "required": false,
+            "schema": {
+              "type": "array",
+              "minItems": 1,
+              "uniqueItems": true,
+              "items": {
+                "type": "object"
+              }
+            }
+          },
+          {
+            "$ref": "#/parameters/api-version"
+          }
+        ],
+        "responses": {
+          "201": {
+            "description": "Success",
+            "schema": {
+              "$ref": "#/definitions/NonPagedDigitalTwinsModelDataCollection"
+            }
+          },
+          "default": {
+            "description": "Default response.",
+            "schema": {
+              "$ref": "#/definitions/ErrorResponse"
+            }
+          }
+        }
+      },
+      "get": {
+        "tags": [
+          "Models"
+        ],
+        "description": "Retrieves model metadata and, optionally, model definitions.\nStatus codes:\n* 200 OK\n* 400 Bad Request\n  * InvalidArgument - The model id is invalid.\n  * LimitExceeded - The maximum number of model ids allowed in 'dependenciesFor' has been reached.\n* 404 Not Found\n  * ModelNotFound - The model was not found.",
+        "x-ms-pageable": {
+          "nextLinkName": "nextLink"
+        },
+        "operationId": "DigitalTwinModels_List",
+        "x-ms-examples": {
+          "List models": {
+            "$ref": "./examples/ModelsGet.json"
+          },
+          "List models (include dependencies and definitions)": {
+            "$ref": "./examples/ModelsGetWithDependenciesAndIncludeModelDefinition.json"
+          }
+        },
+        "parameters": [
+          {
+            "name": "dependenciesFor",
+            "in": "query",
+            "description": "The set of the models which will have their dependencies retrieved. If omitted, all models are retrieved.",
+            "required": false,
+            "type": "array",
+            "collectionFormat": "multi",
+            "items": {
+              "type": "string"
+            }
+          },
+          {
+            "$ref": "#/parameters/includeModelDefinition"
+          },
+          {
+            "$ref": "#/parameters/max-items-per-page"
+          },
+          {
+            "$ref": "#/parameters/api-version"
+          }
+        ],
+        "responses": {
+          "200": {
+            "description": "Success",
+            "schema": {
+              "$ref": "#/definitions/PagedDigitalTwinsModelDataCollection"
+            }
+          },
+          "default": {
+            "description": "Default response.",
+            "schema": {
+              "$ref": "#/definitions/ErrorResponse"
+            }
+          }
+        }
+      }
+    },
+    "/models/{id}": {
+      "parameters": [
+        {
+          "$ref": "#/parameters/traceparent"
+        },
+        {
+          "$ref": "#/parameters/tracestate"
+        }
+      ],
+      "get": {
+        "tags": [
+          "Models"
+        ],
+        "description": "Retrieves model metadata and optionally the model definition.\nStatus codes:\n* 200 OK\n* 400 Bad Request\n  * InvalidArgument - The model id is invalid.\n  * MissingArgument - The model id was not provided.\n* 404 Not Found\n  * ModelNotFound - The model was not found.",
+        "operationId": "DigitalTwinModels_GetById",
+        "x-ms-examples": {
+          "Get a model by id": {
+            "$ref": "./examples/ModelGetById.json"
+          },
+          "Get a model by id (with definition)": {
+            "$ref": "./examples/ModelGetByIdWithIncludeModelDefinition.json"
+          }
+        },
+        "parameters": [
+          {
+            "$ref": "#/parameters/modelId"
+          },
+          {
+            "$ref": "#/parameters/includeModelDefinition"
+          },
+          {
+            "$ref": "#/parameters/api-version"
+          }
+        ],
+        "responses": {
+          "200": {
+            "description": "Success",
+            "schema": {
+              "$ref": "#/definitions/DigitalTwinsModelData"
+            }
+          },
+          "default": {
+            "description": "Default response.",
+            "schema": {
+              "$ref": "#/definitions/ErrorResponse"
+            }
+          }
+        }
+      },
+      "patch": {
+        "tags": [
+          "Models"
+        ],
+        "description": "Updates the metadata for a model.\nStatus codes:\n* 204 No Content\n* 400 Bad Request\n  * InvalidArgument - The model id is invalid.\n  * JsonPatchInvalid - The JSON Patch provided is invalid.\n  * MissingArgument - The model id was not provided.\n* 404 Not Found\n  * ModelNotFound - The model was not found.\n* 409 Conflict\n  * ModelReferencesNotDecommissioned - The model refers to models that are not decommissioned.",
+        "operationId": "DigitalTwinModels_Update",
+        "consumes": [
+          "application/json-patch+json"
+        ],
+        "x-ms-examples": {
+          "Decommission a model": {
+            "$ref": "./examples/ModelUpdateDecommissioned.json"
+          }
+        },
+        "parameters": [
+          {
+            "$ref": "#/parameters/modelId"
+          },
+          {
+            "name": "updateModel",
+            "description": "An update specification described by JSON Patch. Only the decommissioned property can be replaced.",
+            "in": "body",
+            "required": true,
+            "schema": {
+              "type": "array",
+              "items": {
+                "type": "object"
+              }
+            }
+          },
+          {
+            "$ref": "#/parameters/api-version"
+          }
+        ],
+        "responses": {
+          "204": {
+            "description": "Success"
+          },
+          "default": {
+            "description": "Default response.",
+            "schema": {
+              "$ref": "#/definitions/ErrorResponse"
+            }
+          }
+        }
+      },
+      "delete": {
+        "tags": [
+          "Models"
+        ],
+        "operationId": "DigitalTwinModels_Delete",
+        "description": "Deletes a model. A model can only be deleted if no other models reference it.\nStatus codes:\n* 204 No Content\n* 400 Bad Request\n  * InvalidArgument - The model id is invalid.\n  * MissingArgument - The model id was not provided.\n* 404 Not Found\n  * ModelNotFound - The model was not found.\n* 409 Conflict\n  * ModelReferencesNotDeleted - The model refers to models that are not deleted.",
+        "x-ms-examples": {
+          "Delete a model": {
+            "$ref": "./examples/DeleteModel.json"
+          }
+        },
+        "parameters": [
+          {
+            "$ref": "#/parameters/modelId"
+          },
+          {
+            "$ref": "#/parameters/api-version"
+          }
+        ],
+        "responses": {
+          "204": {
+            "description": "Success"
+          },
+          "default": {
+            "description": "Default response.",
+            "schema": {
+              "$ref": "#/definitions/ErrorResponse"
+            }
+          }
+        }
+      }
+    },
+    "/query": {
+      "parameters": [
+        {
+          "$ref": "#/parameters/traceparent"
+        },
+        {
+          "$ref": "#/parameters/tracestate"
+        }
+      ],
+      "post": {
+        "tags": [
+          "Query"
+        ],
+        "operationId": "Query_QueryTwins",
+        "description": "Executes a query that allows traversing relationships and filtering by property values.\nStatus codes:\n* 200 OK\n* 400 Bad Request\n  * BadRequest - The continuation token is invalid.\n  * SqlQueryError - The query contains some errors.\n  * TimeoutError - The query execution timed out after 60 seconds. Try simplifying the query or adding conditions to reduce the result size.\n * 429 Too Many Requests\n  * QuotaReachedError - The maximum query rate limit has been reached.",
+        "x-ms-examples": {
+          "Query for digital twins": {
+            "$ref": "./examples/QueryFirstPage.json"
+          },
+          "Retrieve paged query results": {
+            "$ref": "./examples/QueryNextPage.json"
+          },
+          "Query for digital twins through a relationship": {
+            "$ref": "./examples/QueryJoin.json"
+          }
+        },
+        "parameters": [
+          {
+            "name": "querySpecification",
+            "in": "body",
+            "description": "The query specification to execute.",
+            "required": true,
+            "schema": {
+              "$ref": "#/definitions/QuerySpecification"
+            }
+          },
+          {
+            "$ref": "#/parameters/max-items-per-page"
+          },
+          {
+            "$ref": "#/parameters/api-version"
+          }
+        ],
+        "responses": {
+          "200": {
+            "description": "Success",
+            "schema": {
+              "$ref": "#/definitions/QueryResult"
+            },
+            "headers": {
+              "query-charge": {
+                "description": "The query charge.",
+                "type": "number"
+              }
+            }
+          },
+          "default": {
+            "description": "Default response.",
+            "schema": {
+              "$ref": "#/definitions/ErrorResponse"
+            }
+          }
+        }
+      }
+    },
+    "/digitaltwins/{id}": {
+      "parameters": [
+        {
+          "$ref": "#/parameters/traceparent"
+        },
+        {
+          "$ref": "#/parameters/tracestate"
+        }
+      ],
+      "get": {
+        "tags": [
+          "Twins"
+        ],
+        "operationId": "DigitalTwins_GetById",
+        "description": "Retrieves a digital twin.\nStatus codes:\n* 200 OK\n* 400 Bad Request\n  * InvalidArgument - The digital twin id is invalid.\n* 404 Not Found\n  * DigitalTwinNotFound - The digital twin was not found.",
+        "x-ms-examples": {
+          "Get a digital twin by id": {
+            "$ref": "./examples/GetTwin.json"
+          }
+        },
+        "parameters": [
+          {
+            "$ref": "#/parameters/digitalTwinId"
+          },
+          {
+            "$ref": "#/parameters/api-version"
+          }
+        ],
+        "responses": {
+          "200": {
+            "description": "Success",
+            "schema": {
+              "type": "object"
+            },
+            "headers": {
+              "ETag": {
+                "description": "Weak Etag.",
+                "type": "string"
+              }
+            }
+          },
+          "default": {
+            "description": "Default response.",
+            "schema": {
+              "$ref": "#/definitions/ErrorResponse"
+            }
+          }
+        }
+      },
+      "put": {
+        "tags": [
+          "Twins"
+        ],
+        "operationId": "DigitalTwins_Add",
+        "description": "Adds or replaces a digital twin.\nStatus codes:\n* 200 OK\n* 400 Bad Request\n  * InvalidArgument - The digital twin id or payload is invalid.\n  * ModelDecommissioned - The model for the digital twin is decommissioned.\n  * TwinLimitReached - The maximum number of digital twins allowed has been reached.\n  * ValidationFailed - The digital twin payload is not valid.\n* 412 Precondition Failed\n  * PreconditionFailed - The precondition check (If-Match or If-None-Match) failed.",
+        "x-ms-examples": {
+          "Add a digital twin": {
+            "$ref": "./examples/PutTwinBasicExample.json"
+          },
+          "Add a digital twin (with properties and components)": {
+            "$ref": "./examples/PutTwinAdvancedExample.json"
+          }
+        },
+        "parameters": [
+          {
+            "$ref": "#/parameters/digitalTwinId"
+          },
+          {
+            "name": "twin",
+            "description": "The digital twin instance being added. If provided, the $dtId property is ignored.",
+            "in": "body",
+            "required": true,
+            "schema": {
+              "type": "object"
+            }
+          },
+          {
+            "$ref": "#/parameters/if-none-match-star"
+          },
+          {
+            "$ref": "#/parameters/api-version"
+          }
+        ],
+        "responses": {
+          "200": {
+            "description": "Success",
+            "schema": {
+              "type": "object"
+            },
+            "headers": {
+              "ETag": {
+                "description": "Weak Etag.",
+                "type": "string"
+              }
+            }
+          },
+          "202": {
+            "description": "Asynchronous Success"
+          },
+          "default": {
+            "description": "Default response.",
+            "schema": {
+              "$ref": "#/definitions/ErrorResponse"
+            }
+          }
+        }
+      },
+      "delete": {
+        "tags": [
+          "Twins"
+        ],
+        "operationId": "DigitalTwins_Delete",
+        "description": "Deletes a digital twin. All relationships referencing the digital twin must already be deleted.\nStatus codes:\n* 204 No Content\n* 400 Bad Request\n  * InvalidArgument - The digital twin id is invalid.\n  * RelationshipsNotDeleted - The digital twin contains relationships.\n* 404 Not Found\n  * DigitalTwinNotFound - The digital twin was not found.\n* 412 Precondition Failed\n  * PreconditionFailed - The precondition check (If-Match or If-None-Match) failed.",
+        "x-ms-examples": {
+          "Delete a digital twin": {
+            "$ref": "./examples/DeleteTwin.json"
+          }
+        },
+        "parameters": [
+          {
+            "$ref": "#/parameters/digitalTwinId"
+          },
+          {
+            "$ref": "#/parameters/if-match"
+          },
+          {
+            "$ref": "#/parameters/api-version"
+          }
+        ],
+        "responses": {
+          "204": {
+            "description": "Success"
+          },
+          "default": {
+            "description": "Default response.",
+            "schema": {
+              "$ref": "#/definitions/ErrorResponse"
+            }
+          }
+        }
+      },
+      "patch": {
+        "tags": [
+          "Twins"
+        ],
+        "operationId": "DigitalTwins_Update",
+        "description": "Updates a digital twin.\nStatus codes:\n* 204 No Content\n* 400 Bad Request\n  * InvalidArgument - The digital twin id or payload is invalid.\n  * JsonPatchInvalid - The JSON Patch provided is invalid.\n  * ValidationFailed - Applying the patch results in an invalid digital twin.\n* 404 Not Found\n  * DigitalTwinNotFound - The digital twin was not found.\n* 412 Precondition Failed\n  * PreconditionFailed - The precondition check (If-Match or If-None-Match) failed.",
+        "consumes": [
+          "application/json-patch+json"
+        ],
+        "x-ms-examples": {
+          "Update a digital twin": {
+            "$ref": "./examples/PatchTwinBasicExample.json"
+          },
+          "Update a digital twin (multiple updates)": {
+            "$ref": "./examples/PatchTwinAdvancedExample.json"
+          }
+        },
+        "parameters": [
+          {
+            "$ref": "#/parameters/digitalTwinId"
+          },
+          {
+            "name": "patchDocument",
+            "description": "An update specification described by JSON Patch. Updates to property values and $model elements may happen in the same request. Operations are limited to add, replace and remove.",
+            "in": "body",
+            "required": true,
+            "schema": {
+              "type": "array",
+              "items": {
+                "type": "object"
+              }
+            }
+          },
+          {
+            "$ref": "#/parameters/if-match"
+          },
+          {
+            "$ref": "#/parameters/api-version"
+          }
+        ],
+        "responses": {
+          "204": {
+            "description": "Success",
+            "headers": {
+              "ETag": {
+                "description": "Weak Etag.",
+                "type": "string"
+              }
+            }
+          },
+          "202": {
+            "description": "Asynchronous Success"
+          },
+          "default": {
+            "description": "Default response.",
+            "schema": {
+              "$ref": "#/definitions/ErrorResponse"
+            }
+          }
+        }
+      }
+    },
+    "/digitaltwins/{id}/relationships/{relationshipId}": {
+      "parameters": [
+        {
+          "$ref": "#/parameters/traceparent"
+        },
+        {
+          "$ref": "#/parameters/tracestate"
+        }
+      ],
+      "get": {
+        "tags": [
+          "Twins"
+        ],
+        "operationId": "DigitalTwins_GetRelationshipById",
+        "description": "Retrieves a relationship between two digital twins.\nStatus codes:\n* 200 OK\n* 400 Bad Request\n  * InvalidArgument - The digital twin id or relationship id is invalid.\n* 404 Not Found\n  * DigitalTwinNotFound - The digital twin was not found.\n  * RelationshipNotFound - The relationship was not found.",
+        "x-ms-examples": {
+          "Get a relationship by id": {
+            "$ref": "./examples/GetRelationshipById.json"
+          }
+        },
+        "parameters": [
+          {
+            "$ref": "#/parameters/digitalTwinId"
+          },
+          {
+            "$ref": "#/parameters/relationshipId"
+          },
+          {
+            "$ref": "#/parameters/api-version"
+          }
+        ],
+        "responses": {
+          "200": {
+            "description": "Success",
+            "schema": {
+              "type": "object"
+            },
+            "headers": {
+              "ETag": {
+                "description": "Weak Etag.",
+                "type": "string"
+              }
+            }
+          },
+          "default": {
+            "description": "Default response.",
+            "schema": {
+              "$ref": "#/definitions/ErrorResponse"
+            }
+          }
+        }
+      },
+      "put": {
+        "tags": [
+          "Twins"
+        ],
+        "operationId": "DigitalTwins_AddRelationship",
+        "description": "Adds a relationship between two digital twins.\nStatus codes:\n* 200 OK\n* 400 Bad Request\n  * InvalidArgument - The digital twin id, relationship id, or payload is invalid.\n  * InvalidRelationship - The relationship is invalid.\n  * OperationNotAllowed - The relationship cannot connect to the same digital twin.\n  * ValidationFailed - The relationship content is invalid.\n* 404 Not Found\n  * DigitalTwinNotFound - The digital twin was not found.\n  * TargetTwinNotFound - The digital twin target of the relationship was not found.\n* 412 Precondition Failed\n  * PreconditionFailed - The precondition check (If-Match or If-None-Match) failed.",
+        "x-ms-examples": {
+          "Create a relationship": {
+            "$ref": "./examples/CreateRelationshipBasicExample.json"
+          },
+          "Create a relationship (with properties)": {
+            "$ref": "./examples/CreateRelationshipAdvancedExample.json"
+          }
+        },
+        "parameters": [
+          {
+            "$ref": "#/parameters/digitalTwinId"
+          },
+          {
+            "$ref": "#/parameters/relationshipId"
+          },
+          {
+            "name": "relationship",
+            "in": "body",
+            "description": "The data for the relationship.",
+            "required": true,
+            "schema": {
+              "type": "object"
+            }
+          },
+          {
+            "$ref": "#/parameters/if-none-match-star"
+          },
+          {
+            "$ref": "#/parameters/api-version"
+          }
+        ],
+        "responses": {
+          "200": {
+            "description": "Success",
+            "schema": {
+              "type": "object"
+            },
+            "headers": {
+              "ETag": {
+                "description": "Weak Etag.",
+                "type": "string"
+              }
+            }
+          },
+          "default": {
+            "description": "Default response.",
+            "schema": {
+              "$ref": "#/definitions/ErrorResponse"
+            }
+          }
+        }
+      },
+      "delete": {
+        "tags": [
+          "Twins"
+        ],
+        "operationId": "DigitalTwins_DeleteRelationship",
+        "description": "Deletes a relationship between two digital twins.\nStatus codes:\n* 204 No Content\n* 400 Bad Request\n  * InvalidArgument - The digital twin id or relationship id is invalid.\n* 404 Not Found\n  * DigitalTwinNotFound - The digital twin was not found.\n  * RelationshipNotFound - The relationship was not found.\n* 412 Precondition Failed\n  * PreconditionFailed - The precondition check (If-Match or If-None-Match) failed.",
+        "x-ms-examples": {
+          "Delete a relationship": {
+            "$ref": "./examples/DeleteRelationship.json"
+          }
+        },
+        "parameters": [
+          {
+            "$ref": "#/parameters/digitalTwinId"
+          },
+          {
+            "$ref": "#/parameters/relationshipId"
+          },
+          {
+            "$ref": "#/parameters/if-match"
+          },
+          {
+            "$ref": "#/parameters/api-version"
+          }
+        ],
+        "responses": {
+          "204": {
+            "description": "Success"
+          },
+          "default": {
+            "description": "Default response.",
+            "schema": {
+              "$ref": "#/definitions/ErrorResponse"
+            }
+          }
+        }
+      },
+      "patch": {
+        "tags": [
+          "Twins"
+        ],
+        "consumes": [
+          "application/json-patch+json"
+        ],
+        "operationId": "DigitalTwins_UpdateRelationship",
+        "description": "Updates the properties on a relationship between two digital twins.\nStatus codes:\n* 204 No Content\n* 400 Bad Request\n  * InvalidArgument - The digital twin id or relationship id is invalid.\n  * InvalidRelationship - The relationship is invalid.\n  * JsonPatchInvalid - The JSON Patch provided is invalid.\n  * ValidationFailed - The relationship content is invalid.\n* 404 Not Found\n  * DigitalTwinNotFound - The digital twin was not found.\n  * RelationshipNotFound - The relationship was not found.\n* 409 Conflict\n  * RelationshipAlreadyExists - The relationship already exists.\n* 412 Precondition Failed\n  * PreconditionFailed - The precondition check (If-Match or If-None-Match) failed.",
+        "x-ms-examples": {
+          "Update relationship properties": {
+            "$ref": "./examples/PatchRelationshipBasicExample.json"
+          },
+          "Update relationship properties (multiple updates)": {
+            "$ref": "./examples/PatchRelationshipAdvancedExample.json"
+          }
+        },
+        "parameters": [
+          {
+            "$ref": "#/parameters/digitalTwinId"
+          },
+          {
+            "$ref": "#/parameters/relationshipId"
+          },
+          {
+            "name": "patchDocument",
+            "description": "JSON Patch description of the update to the relationship properties.",
+            "in": "body",
+            "required": true,
+            "schema": {
+              "type": "array",
+              "items": {
+                "type": "object"
+              }
+            }
+          },
+          {
+            "$ref": "#/parameters/if-match"
+          },
+          {
+            "$ref": "#/parameters/api-version"
+          }
+        ],
+        "responses": {
+          "204": {
+            "description": "Success",
+            "headers": {
+              "ETag": {
+                "description": "Weak Etag.",
+                "type": "string"
+              }
+            }
+          },
+          "default": {
+            "description": "Default response.",
+            "schema": {
+              "$ref": "#/definitions/ErrorResponse"
+            }
+          }
+        }
+      }
+    },
+    "/digitaltwins/{id}/relationships": {
+      "parameters": [
+        {
+          "$ref": "#/parameters/traceparent"
+        },
+        {
+          "$ref": "#/parameters/tracestate"
+        }
+      ],
+      "get": {
+        "tags": [
+          "Twins"
+        ],
+        "operationId": "DigitalTwins_ListRelationships",
+        "description": "Retrieves the relationships from a digital twin.\nStatus codes:\n* 200 OK\n* 400 Bad Request\n  * InvalidArgument - The digital twin id is invalid.\n* 404 Not Found\n  * DigitalTwinNotFound - The digital twin was not found.",
+        "x-ms-pageable": {
+          "nextLinkName": "nextLink"
+        },
+        "x-ms-examples": {
+          "List relationships": {
+            "$ref": "./examples/GetRelationship.json"
+          },
+          "List relationships by name": {
+            "$ref": "./examples/GetRelationshipByRelationshipName.json"
+          }
+        },
+        "parameters": [
+          {
+            "$ref": "#/parameters/digitalTwinId"
+          },
+          {
+            "name": "relationshipName",
+            "description": "The name of the relationship.",
+            "in": "query",
+            "required": false,
+            "type": "string"
+          },
+          {
+            "$ref": "#/parameters/api-version"
+          }
+        ],
+        "responses": {
+          "200": {
+            "description": "Success",
+            "schema": {
+              "$ref": "#/definitions/RelationshipCollection"
+            }
+          },
+          "default": {
+            "description": "Default response.",
+            "schema": {
+              "$ref": "#/definitions/ErrorResponse"
+            }
+          }
+        }
+      }
+    },
+    "/digitaltwins/{id}/incomingrelationships": {
+      "parameters": [
+        {
+          "$ref": "#/parameters/traceparent"
+        },
+        {
+          "$ref": "#/parameters/tracestate"
+        }
+      ],
+      "get": {
+        "tags": [
+          "Twins"
+        ],
+        "operationId": "DigitalTwins_ListIncomingRelationships",
+        "description": "Retrieves all incoming relationship for a digital twin.\nStatus codes:\n* 200 OK\n* 400 Bad Request\n  * InvalidArgument - The digital twin id is invalid.\n* 404 Not Found\n  * DigitalTwinNotFound - The digital twin was not found.",
+        "x-ms-pageable": {
+          "nextLinkName": "nextLink"
+        },
+        "x-ms-examples": {
+          "List incoming relationships": {
+            "$ref": "./examples/GetIncomingRelationship.json"
+          }
+        },
+        "parameters": [
+          {
+            "$ref": "#/parameters/digitalTwinId"
+          },
+          {
+            "$ref": "#/parameters/api-version"
+          }
+        ],
+        "responses": {
+          "200": {
+            "description": "Success",
+            "schema": {
+              "$ref": "#/definitions/IncomingRelationshipCollection"
+            }
+          },
+          "default": {
+            "description": "Default response.",
+            "schema": {
+              "$ref": "#/definitions/ErrorResponse"
+            }
+          }
+        }
+      }
+    },
+    "/digitaltwins/{id}/telemetry": {
+      "parameters": [
+        {
+          "$ref": "#/parameters/traceparent"
+        },
+        {
+          "$ref": "#/parameters/tracestate"
+        }
+      ],
+      "post": {
+        "tags": [
+          "Twins"
+        ],
+        "operationId": "DigitalTwins_SendTelemetry",
+        "description": "Sends telemetry on behalf of a digital twin.\nStatus codes:\n* 204 No Content\n* 400 Bad Request\n  * InvalidArgument - The digital twin id or message id is invalid.\n  * ValidationFailed - The telemetry content is invalid.\n* 404 Not Found\n  * DigitalTwinNotFound - The digital twin was not found.",
+        "x-ms-examples": {
+          "Send telemetry": {
+            "$ref": "./examples/SendTelemetry.json"
+          }
+        },
+        "parameters": [
+          {
+            "$ref": "#/parameters/digitalTwinId"
+          },
+          {
+            "name": "telemetry",
+            "in": "body",
+            "description": "The telemetry measurements to send from the digital twin.",
+            "required": true,
+            "schema": {
+              "type": "object"
+            }
+          },
+          {
+            "name": "Message-Id",
+            "in": "header",
+            "description": "A unique message identifier (in the scope of the digital twin id) that is commonly used for de-duplicating messages.",
+            "required": true,
+            "type": "string"
+          },
+          {
+            "name": "Telemetry-Source-Time",
+            "in": "header",
+            "description": "An RFC 3339 timestamp that identifies the time the telemetry was measured.",
+            "required": false,
+            "type": "string"
+          },
+          {
+            "$ref": "#/parameters/api-version"
+          }
+        ],
+        "responses": {
+          "204": {
+            "description": "Success"
+          },
+          "default": {
+            "description": "Default response.",
+            "schema": {
+              "$ref": "#/definitions/ErrorResponse"
+            }
+          }
+        }
+      }
+    },
+    "/digitaltwins/{id}/components/{componentPath}/telemetry": {
+      "parameters": [
+        {
+          "$ref": "#/parameters/traceparent"
+        },
+        {
+          "$ref": "#/parameters/tracestate"
+        }
+      ],
+      "post": {
+        "tags": [
+          "Twins"
+        ],
+        "operationId": "DigitalTwins_SendComponentTelemetry",
+        "description": "Sends telemetry on behalf of a component in a digital twin.\nStatus codes:\n* 204 No Content\n* 400 Bad Request\n  * InvalidArgument - The digital twin id, message id, or component path is invalid.\n  * ValidationFailed - The telemetry content is invalid.\n* 404 Not Found\n  * DigitalTwinNotFound - The digital twin was not found.\n  * ComponentNotFound - The component path was not found.",
+        "x-ms-examples": {
+          "Send telemetry in a component": {
+            "$ref": "./examples/SendTelemetryFromComponent.json"
+          }
+        },
+        "parameters": [
+          {
+            "$ref": "#/parameters/digitalTwinId"
+          },
+          {
+            "$ref": "#/parameters/componentPath"
+          },
+          {
+            "name": "telemetry",
+            "in": "body",
+            "description": "The telemetry measurements to send from the digital twin's component.",
+            "required": true,
+            "schema": {
+              "type": "object"
+            }
+          },
+          {
+            "name": "Message-Id",
+            "in": "header",
+            "description": "A unique message identifier (in the scope of the digital twin id) that is commonly used for de-duplicating messages.",
+            "required": true,
+            "type": "string"
+          },
+          {
+            "name": "Telemetry-Source-Time",
+            "in": "header",
+            "description": "An RFC 3339 timestamp that identifies the time the telemetry was measured.",
+            "required": false,
+            "type": "string"
+          },
+          {
+            "$ref": "#/parameters/api-version"
+          }
+        ],
+        "responses": {
+          "204": {
+            "description": "Success"
+          },
+          "default": {
+            "description": "Default response.",
+            "schema": {
+              "$ref": "#/definitions/ErrorResponse"
+            }
+          }
+        }
+      }
+    },
+    "/digitaltwins/{id}/components/{componentPath}": {
+      "parameters": [
+        {
+          "$ref": "#/parameters/traceparent"
+        },
+        {
+          "$ref": "#/parameters/tracestate"
+        }
+      ],
+      "get": {
+        "tags": [
+          "Twins"
+        ],
+        "operationId": "DigitalTwins_GetComponent",
+        "description": "Retrieves a component from a digital twin.\nStatus codes:\n* 200 OK\n* 400 Bad Request\n  * InvalidArgument - The digital twin id or component path is invalid.\n* 404 Not Found\n  * DigitalTwinNotFound - The digital twin was not found.\n  * ComponentNotFound - The component path was not found.",
+        "x-ms-examples": {
+          "Get a component in a digital twin": {
+            "$ref": "./examples/GetComponent.json"
+          }
+        },
+        "parameters": [
+          {
+            "$ref": "#/parameters/digitalTwinId"
+          },
+          {
+            "$ref": "#/parameters/componentPath"
+          },
+          {
+            "$ref": "#/parameters/api-version"
+          }
+        ],
+        "responses": {
+          "200": {
+            "description": "Success",
+            "schema": {
+              "type": "object"
+            },
+            "headers": {
+              "ETag": {
+                "description": "Weak Etag.",
+                "type": "string"
+              }
+            }
+          },
+          "default": {
+            "description": "Default response.",
+            "schema": {
+              "$ref": "#/definitions/ErrorResponse"
+            }
+          }
+        }
+      },
+      "patch": {
+        "tags": [
+          "Twins"
+        ],
+        "operationId": "DigitalTwins_UpdateComponent",
+        "description": "Updates a component on a digital twin.\nStatus codes:\n* 204 No Content\n* 400 Bad Request\n  * InvalidArgument - The digital twin id, component path, or payload is invalid.\n  * JsonPatchInvalid - The JSON Patch provided is invalid.\n  * ValidationFailed - Applying the patch results in an invalid digital twin.\n* 404 Not Found\n  * DigitalTwinNotFound - The digital twin was not found.\n* 412 Precondition Failed\n  * PreconditionFailed - The precondition check (If-Match or If-None-Match) failed.",
+        "x-ms-examples": {
+          "Update a component in a digital twin": {
+            "$ref": "./examples/PatchComponentBasicExample.json"
+          },
+          "Update a component in a digital twin (multiple updates)": {
+            "$ref": "./examples/PatchComponentAdvancedExample.json"
+          }
+        },
+        "consumes": [
+          "application/json-patch+json"
+        ],
+        "parameters": [
+          {
+            "$ref": "#/parameters/digitalTwinId"
+          },
+          {
+            "$ref": "#/parameters/componentPath"
+          },
+          {
+            "name": "patchDocument",
+            "description": "An update specification described by JSON Patch. Updates to property values and $model elements may happen in the same request. Operations are limited to add, replace and remove.",
+            "in": "body",
+            "required": true,
+            "schema": {
+              "type": "array",
+              "items": {
+                "type": "object"
+              }
+            }
+          },
+          {
+            "$ref": "#/parameters/if-match"
+          },
+          {
+            "$ref": "#/parameters/api-version"
+          }
+        ],
+        "responses": {
+          "204": {
+            "description": "Success",
+            "headers": {
+              "ETag": {
+                "description": "Weak Etag.",
+                "type": "string"
+              }
+            }
+          },
+          "202": {
+            "description": "Asynchronous Success"
+          },
+          "default": {
+            "description": "Default response.",
+            "schema": {
+              "$ref": "#/definitions/ErrorResponse"
+            }
+          }
+        }
+      }
+    },
+    "/eventroutes": {
+      "parameters": [
+        {
+          "$ref": "#/parameters/traceparent"
+        },
+        {
+          "$ref": "#/parameters/tracestate"
+        }
+      ],
+      "get": {
+        "tags": [
+          "EventRoutes"
+        ],
+        "description": "Retrieves all event routes.\nStatus codes:\n* 200 OK",
+        "operationId": "EventRoutes_List",
+        "x-ms-examples": {
+          "List event routes": {
+            "$ref": "./examples/EventRoutesList.json"
+          }
+        },
+        "parameters": [
+          {
+            "$ref": "#/parameters/max-items-per-page"
+          },
+          {
+            "$ref": "#/parameters/api-version"
+          }
+        ],
+        "responses": {
+          "200": {
+            "description": "Success",
+            "schema": {
+              "$ref": "#/definitions/EventRouteCollection"
+            }
+          },
+          "default": {
+            "description": "Default response.",
+            "schema": {
+              "$ref": "#/definitions/ErrorResponse"
+            }
+          }
+        },
+        "x-ms-pageable": {
+          "nextLinkName": "nextLink"
+        }
+      }
+    },
+    "/eventroutes/{id}": {
+      "parameters": [
+        {
+          "$ref": "#/parameters/traceparent"
+        },
+        {
+          "$ref": "#/parameters/tracestate"
+        }
+      ],
+      "get": {
+        "tags": [
+          "EventRoutes"
+        ],
+        "description": "Retrieves an event route.\nStatus codes:\n* 200 OK\n* 404 Not Found\n  * EventRouteNotFound - The event route was not found.",
+        "operationId": "EventRoutes_GetById",
+        "x-ms-examples": {
+          "Get an event route by id": {
+            "$ref": "./examples/EventRouteGet.json"
+          }
+        },
+        "parameters": [
+          {
+            "$ref": "#/parameters/eventRouteId"
+          },
+          {
+            "$ref": "#/parameters/api-version"
+          }
+        ],
+        "responses": {
+          "200": {
+            "description": "Success",
+            "schema": {
+              "$ref": "#/definitions/EventRoute"
+            }
+          },
+          "default": {
+            "description": "Default response.",
+            "schema": {
+              "$ref": "#/definitions/ErrorResponse"
+            }
+          }
+        }
+      },
+      "put": {
+        "tags": [
+          "EventRoutes"
+        ],
+        "description": "Adds or replaces an event route.\nStatus codes:\n* 204 No Content\n* 400 Bad Request\n  * EventRouteEndpointInvalid - The endpoint provided does not exist or is not active.\n  * EventRouteFilterInvalid - The event route filter is invalid.\n  * EventRouteIdInvalid - The event route id is invalid.\n  * LimitExceeded - The maximum number of event routes allowed has been reached.",
+        "operationId": "EventRoutes_Add",
+        "x-ms-examples": {
+          "Add an event route": {
+            "$ref": "./examples/EventRoutePut.json"
+          },
+          "Add an event route (with filter)": {
+            "$ref": "./examples/EventRoutePutWithFilter.json"
+          }
+        },
+        "parameters": [
+          {
+            "$ref": "#/parameters/eventRouteId"
+          },
+          {
+            "name": "eventRoute",
+            "description": "The event route data",
+            "in": "body",
+            "schema": {
+              "$ref": "#/definitions/EventRoute"
+            }
+          },
+          {
+            "$ref": "#/parameters/api-version"
+          }
+        ],
+        "responses": {
+          "204": {
+            "description": "Success"
+          },
+          "default": {
+            "description": "Default response.",
+            "schema": {
+              "$ref": "#/definitions/ErrorResponse"
+            }
+          }
+        }
+      },
+      "delete": {
+        "tags": [
+          "EventRoutes"
+        ],
+        "description": "Deletes an event route.\nStatus codes:\n* 204 No Content\n* 404 Not Found\n  * EventRouteNotFound - The event route was not found.",
+        "operationId": "EventRoutes_Delete",
+        "x-ms-examples": {
+          "Delete an event route": {
+            "$ref": "./examples/EventRouteDelete.json"
+          }
+        },
+        "parameters": [
+          {
+            "$ref": "#/parameters/eventRouteId"
+          },
+          {
+            "$ref": "#/parameters/api-version"
+          }
+        ],
+        "responses": {
+          "204": {
+            "description": "Success"
+          },
+          "default": {
+            "description": "Default response.",
+            "schema": {
+              "$ref": "#/definitions/ErrorResponse"
+            }
+          }
+        }
+      }
+    }
+  },
+  "definitions": {
+    "EventRoute": {
+      "description": "A route which directs notification and telemetry events to an endpoint. Endpoints are a destination outside of Azure Digital Twins such as an EventHub.",
+      "type": "object",
+      "required": [
+        "endpointName",
+        "filter"
+      ],
+      "properties": {
+        "id": {
+          "description": "The id of the event route.",
+          "type": "string",
+          "readOnly": true
+        },
+        "endpointName": {
+          "description": "The name of the endpoint this event route is bound to.",
+          "type": "string"
+        },
+        "filter": {
+          "description": "An expression which describes the events which are routed to the endpoint.",
+          "type": "string"
+        }
+      }
+    },
+    "EventRouteCollection": {
+      "description": "A collection of EventRoute objects.",
+      "type": "object",
+      "properties": {
+        "value": {
+          "description": "The EventRoute objects.",
+          "type": "array",
+          "items": {
+            "$ref": "#/definitions/EventRoute"
+          }
+        },
+        "nextLink": {
+          "description": "A URI to retrieve the next page of results.",
+          "type": "string"
+        }
+      }
+    },
+    "NonPagedDigitalTwinsModelDataCollection": {
+      "description": "A collection of DigitalTwinsModelData objects.",
+      "type": "array",
+      "items": {
+        "$ref": "#/definitions/DigitalTwinsModelData"
+      }
+    },
+    "PagedDigitalTwinsModelDataCollection": {
+      "description": "A collection of DigitalTwinsModelData objects.",
+      "type": "object",
+      "properties": {
+        "value": {
+          "description": "The DigitalTwinsModelData objects.",
+          "type": "array",
+          "items": {
+            "$ref": "#/definitions/DigitalTwinsModelData"
+          }
+        },
+        "nextLink": {
+          "description": "A URI to retrieve the next page of objects.",
+          "type": "string"
+        }
+      }
+    },
+    "DigitalTwinsModelData": {
+      "description": "A model definition and metadata for that model.",
+      "required": [
+        "id"
+      ],
+      "type": "object",
+      "properties": {
+        "displayName": {
+          "description": "A language map that contains the localized display names as specified in the model definition.",
+          "type": "object",
+          "additionalProperties": {
+            "type": "string"
+          }
+        },
+        "description": {
+          "description": "A language map that contains the localized descriptions as specified in the model definition.",
+          "type": "object",
+          "additionalProperties": {
+            "type": "string"
+          }
+        },
+        "id": {
+          "description": "The id of the model as specified in the model definition.",
+          "type": "string"
+        },
+        "uploadTime": {
+          "description": "The time the model was uploaded to the service.",
+          "format": "date-time",
+          "type": "string"
+        },
+        "decommissioned": {
+          "description": "Indicates if the model is decommissioned. Decommissioned models cannot be referenced by newly created digital twins.",
+          "type": "boolean",
+          "default": false
+        },
+        "model": {
+          "description": "The model definition.",
+          "type": "object"
+        }
+      }
+    },
+    "RelationshipCollection": {
+      "description": "A collection of relationships which relate digital twins together.",
+      "type": "object",
+      "properties": {
+        "value": {
+          "description": "The relationship objects.",
+          "type": "array",
+          "items": {
+            "description": "The relationship JSON document.",
+            "type": "object"
+          }
+        },
+        "nextLink": {
+          "description": "A URI to retrieve the next page of objects.",
+          "type": "string"
+        }
+      }
+    },
+    "IncomingRelationshipCollection": {
+      "description": "A collection of incoming relationships which relate digital twins together.",
+      "type": "object",
+      "properties": {
+        "value": {
+          "type": "array",
+          "items": {
+            "$ref": "#/definitions/IncomingRelationship"
+          }
+        },
+        "nextLink": {
+          "type": "string",
+          "description": "A URI to retrieve the next page of objects."
+        }
+      }
+    },
+    "IncomingRelationship": {
+      "type": "object",
+      "description": "An incoming relationship.",
+      "properties": {
+        "$relationshipId": {
+          "type": "string",
+          "description": "A user-provided string representing the id of this relationship, unique in the context of the source digital twin, i.e. sourceId + relationshipId is unique in the context of the service."
+        },
+        "$sourceId": {
+          "type": "string",
+          "description": "The id of the source digital twin."
+        },
+        "$relationshipName": {
+          "type": "string",
+          "description": "The name of the relationship."
+        },
+        "$relationshipLink": {
+          "type": "string",
+          "description": "Link to the relationship, to be used for deletion."
+        }
+      }
+    },
+    "QuerySpecification": {
+      "description": "A query specification containing either a query statement or a continuation token from a previous query result.",
+      "type": "object",
+      "properties": {
+        "query": {
+          "description": "The query to execute. This value is ignored if a continuation token is provided.",
+          "type": "string"
+        },
+        "continuationToken": {
+          "description": "A token which is used to retrieve the next set of results from a previous query.",
+          "type": "string"
+        }
+      }
+    },
+    "QueryResult": {
+      "description": "The results of a query operation and an optional continuation token.",
+      "type": "object",
+      "properties": {
+        "value": {
+          "description": "The query results.",
+          "type": "array",
+          "items": {
+            "type": "object"
+          }
+        },
+        "continuationToken": {
+          "description": "A token which can be used to construct a new QuerySpecification to retrieve the next set of results.",
+          "type": "string"
+        }
+      }
+    },
+    "ErrorResponse": {
+      "description": "Error response.",
+      "properties": {
+        "error": {
+          "$ref": "#/definitions/Error",
+          "description": "The error details."
+        }
+      }
+    },
+    "Error": {
+      "description": "Error definition.",
+      "properties": {
+        "code": {
+          "description": "Service specific error code which serves as the substatus for the HTTP error code.",
+          "type": "string",
+          "readOnly": true
+        },
+        "message": {
+          "description": "A human-readable representation of the error.",
+          "type": "string",
+          "readOnly": true
+        },
+        "details": {
+          "description": "Internal error details.",
+          "type": "array",
+          "items": {
+            "$ref": "#/definitions/Error"
+          },
+          "readOnly": true
+        },
+        "innererror": {
+          "description": "An object containing more specific information than the current object about the error.",
+          "$ref": "#/definitions/InnerError"
+        }
+      }
+    },
+    "InnerError": {
+      "description": "A more specific error description than was provided by the containing error.",
+      "properties": {
+        "code": {
+          "description": "A more specific error code than was provided by the containing error.",
+          "type": "string"
+        },
+        "innererror": {
+          "description": "An object containing more specific information than the current object about the error.",
+          "$ref": "#/definitions/InnerError"
+        }
+      }
+    }
+  },
+  "securityDefinitions": {
+    "oauth2": {
+      "flow": "implicit",
+      "authorizationUrl": "https://login.microsoftonline.com/common/oauth2/authorize?resource=0b07f429-9f4b-4714-9392-cc5e8e80c8b0",
+      "type": "oauth2"
+    }
+  },
+  "parameters": {
+    "eventRouteId": {
+      "name": "id",
+      "in": "path",
+      "required": true,
+      "type": "string",
+      "description": "The id for an event route. The id is unique within event routes and case sensitive.",
+      "x-ms-parameter-location": "method"
+    },
+    "includeModelDefinition": {
+      "name": "includeModelDefinition",
+      "description": "When true the model definition will be returned as part of the result.",
+      "in": "query",
+      "required": false,
+      "type": "boolean",
+      "default": false,
+      "x-ms-parameter-location": "method"
+    },
+    "digitalTwinId": {
+      "name": "id",
+      "in": "path",
+      "description": "The id of the digital twin. The id is unique within the service and case sensitive.",
+      "type": "string",
+      "required": true,
+      "x-ms-parameter-location": "method"
+    },
+    "modelId": {
+      "name": "id",
+      "in": "path",
+      "description": "The id for the model. The id is globally unique and case sensitive.",
+      "required": true,
+      "type": "string",
+      "x-ms-parameter-location": "method"
+    },
+    "relationshipId": {
+      "name": "relationshipId",
+      "description": "The id of the relationship. The id is unique within the digital twin and case sensitive.",
+      "in": "path",
+      "required": true,
+      "type": "string",
+      "x-ms-parameter-location": "method"
+    },
+    "componentPath": {
+      "name": "componentPath",
+      "in": "path",
+      "description": "The name of the DTDL component.",
+      "type": "string",
+      "required": true,
+      "x-ms-parameter-location": "method"
+    },
+    "api-version": {
+      "name": "api-version",
+      "in": "query",
+      "description": "The requested API version.",
+      "required": true,
+      "type": "string",
+      "enum": [
+        "2020-10-31"
+      ]
+    },
+    "max-items-per-page": {
+      "name": "max-items-per-page",
+      "in": "header",
+      "description": "The maximum number of items to retrieve per request. The server may choose to return less than the requested number.",
+      "required": false,
+      "type": "integer",
+      "x-ms-client-name": "MaxItemsPerPage",
+      "x-ms-parameter-location": "method",
+      "x-ms-parameter-grouping": {
+        "postfix": "Options"
+      }
+    },
+    "if-none-match-star": {
+      "name": "If-None-Match",
+      "in": "header",
+      "description": "Only perform the operation if the entity does not already exist.",
+      "required": false,
+      "type": "string",
+      "enum": [
+        "*"
+      ],
+      "x-ms-parameter-location": "method",
+      "x-ms-parameter-grouping": {
+        "postfix": "Options"
+      }
+    },
+    "if-match": {
+      "name": "If-Match",
+      "in": "header",
+      "description": "Only perform the operation if the entity's etag matches one of the etags provided or * is provided.",
+      "required": false,
+      "type": "string",
+      "x-ms-parameter-location": "method",
+      "x-ms-parameter-grouping": {
+        "postfix": "Options"
+      }
+    },
+    "traceparent": {
+      "name": "traceparent",
+      "in": "header",
+      "description": "Identifies the request in a distributed tracing system.",
+      "required": false,
+      "type": "string",
+      "x-ms-parameter-location": "method",
+      "x-ms-parameter-grouping": {
+        "postfix": "Options"
+      }
+    },
+    "tracestate": {
+      "name": "tracestate",
+      "in": "header",
+      "description": "Provides vendor-specific trace identification information and is a companion to traceparent.",
+      "required": false,
+      "type": "string",
+      "x-ms-parameter-location": "method",
+      "x-ms-parameter-grouping": {
+        "postfix": "Options"
+      }
+    }
+  },
+  "security": [
+    {
+      "oauth2": []
+    }
+  ]
+}
