WebAPI missing files pt2

diff --git a/Engine/Plugins/Web/WebAPI/Source/WebAPIOpenAPI/Private/Tests/Samples/openrouteservice.json b/Engine/Plugins/Web/WebAPI/Source/WebAPIOpenAPI/Private/Tests/Samples/openrouteservice.json
new file mode 100644
--- /dev/null
+++ b/Engine/Plugins/Web/WebAPI/Source/WebAPIOpenAPI/Private/Tests/Samples/openrouteservice.json
@@ -0,0 +1,2098 @@
+{
+  "swagger":"2.0",
+  "info":{
+    "description":"This is the openrouteservice API V2 documentation for ORS Core-Version `5.0.2`. Documentations for [older Core-Versions]([REDACTED_PROFILE_URL]/openrouteservice-docs/releases) can be rendered with the [Swagger-Editor](https://editor.swagger.io/).",
+    "title":"OpenRouteService",
+    "version": "6.6.3",
+    "contact":{
+      "email":"[REDACTED_EMAIL]"
+    },
+    "license":{
+      "name":"MIT",
+      "url":"[REDACTED_PROFILE_URL]/swagger-ui/blob/master/LICENSE"
+    }
+  },
+  "host":"api.openrouteservice.org",
+  "basePath":"/v2",
+  "schemes": [
+    "https"
+  ],
+  "securityDefinitions": {
+    "api_key": {
+        "type": "apiKey",
+        "name": "api_key",
+        "in": "header"
+    }
+},
+  "tags":[
+    {
+      "name":"Directions",
+      "description":"Get directions for different modes of transport"
+    },
+    {
+      "name":"Isochrones",
+      "description":"Obtain areas of reachability from given locations"
+    },
+    {
+      "name":"Matrix",
+      "description":"Obtain one-to-many, many-to-one and many-to-many matrices for time and distance"
+    },
+    {
+      "name":"health-api",
+      "description":"Health API"
+    },
+    {
+      "name":"status-api",
+      "description":"Status API"
+    }
+  ],
+  "paths":{
+    "/v2/directions/{profile}":{
+      "get":{
+        "tags":[
+          "Directions"
+        ],
+        "summary":"Directions Service (GET)",
+        "description":"Get a basic route between two points with the profile provided. Returned response is in GeoJSON format. This method does not accept any request body or parameters other than profile, start coordinate, and end coordinate.",
+        "operationId":"getSimpleGeoJsonRouteUsingGET",
+        "produces":[
+          "application/geo+json;charset=UTF-8"
+        ],
+        "parameters":[
+          {
+            "name":"end",
+            "in":"query",
+            "description":"Destination coordinate of the route",
+            "required":true,
+            "type":"string",
+            "allowEmptyValue":false,
+            "x-example":"8.687872,49.420318"
+          },
+          {
+            "name":"profile",
+            "in":"path",
+            "description":"Specifies the route profile.",
+            "required":true,
+            "type":"string",
+            "x-example":"driving-car",
+            "enum":[
+              "driving-car",
+              "driving-hgv",
+              "cycling-regular",
+              "cycling-road",
+              "cycling-mountain",
+              "cycling-electric",
+              "foot-walking",
+              "foot-hiking",
+              "wheelchair"
+            ]
+          },
+          {
+            "name":"start",
+            "in":"query",
+            "description":"Start coordinate of the route",
+            "required":true,
+            "type":"string",
+            "allowEmptyValue":false,
+            "x-example":"8.681495,49.41461"
+          }
+        ],
+        "responses":{
+          "200":{
+            "description":"Standard response for successfully processed requests. Returns GeoJSON. The decoded values of the extra information can be found [here]([REDACTED_PROFILE_URL]/openrouteservice-docs).",
+            "schema":{
+              "$ref":"#/definitions/GeoJSONRouteResponse"
+            }
+          },
+          "400":{
+            "description":"The request is incorrect and therefore can not be processed."
+          },
+          "401":{
+            "description":"Unauthorized"
+          },
+          "403":{
+            "description":"Forbidden"
+          },
+          "404":{
+            "description":"An element could not be found. If possible, a more detailed error code is provided."
+          },
+          "405":{
+            "description":"The specified HTTP method is not supported. For more details, refer to the EndPoint documentation."
+          },
+          "413":{
+            "description":"The request is larger than the server is able to process, the data provided in the request exceeds the capacity limit."
+          },
+          "500":{
+            "description":"An unexpected error was encountered and a more detailed error code is provided."
+          },
+          "501":{
+            "description":"Indicates that the server does not support the functionality needed to fulfill the request."
+          },
+          "503":{
+            "description":"The server is currently unavailable due to overload or maintenance."
+          }
+        },
+        "deprecated":false
+      },
+      "post":{
+        "tags":[
+          "Directions"
+        ],
+        "summary":"Directions Service (POST)",
+        "description":"Returns a route between two or more locations for a selected profile and its settings as JSON",
+        "operationId":"getDefaultUsingPOST_1",
+        "consumes":[
+          "application/json"
+        ],
+        "produces":[
+          "*/*",
+          "application/json"
+        ],
+        "parameters":[
+          {
+            "name":"profile",
+            "in":"path",
+            "description":"Specifies the route profile.",
+            "required":true,
+            "type":"string",
+            "x-example":"driving-car",
+            "enum":[
+              "driving-car",
+              "driving-hgv",
+              "cycling-regular",
+              "cycling-road",
+              "cycling-mountain",
+              "cycling-electric",
+              "foot-walking",
+              "foot-hiking",
+              "wheelchair"
+            ]
+          },
+          {
+            "in":"body",
+            "name":"request",
+            "description":"The request payload",
+            "required":true,
+            "schema":{
+              "$ref":"#/definitions/Directions Service"
+            }
+          }
+        ],
+        "responses":{
+          "200":{
+            "description":"Standard response for successfully processed requests. Returns JSON. The decoded values of the extra information can be found [here]([REDACTED_PROFILE_URL]/openrouteservice-docs).",
+            "schema":{
+              "$ref":"#/definitions/JSONRouteResponse"
+            }
+          },
+          "201":{
+            "description":"Created"
+          },
+          "400":{
+            "description":"The request is incorrect and therefore can not be processed."
+          },
+          "401":{
+            "description":"Unauthorized"
+          },
+          "403":{
+            "description":"Forbidden"
+          },
+          "404":{
+            "description":"An element could not be found. If possible, a more detailed error code is provided."
+          },
+          "405":{
+            "description":"The specified HTTP method is not supported. For more details, refer to the EndPoint documentation."
+          },
+          "413":{
+            "description":"The request is larger than the server is able to process, the data provided in the request exceeds the capacity limit."
+          },
+          "500":{
+            "description":"An unexpected error was encountered and a more detailed error code is provided."
+          },
+          "501":{
+            "description":"Indicates that the server does not support the functionality needed to fulfill the request."
+          },
+          "503":{
+            "description":"The server is currently unavailable due to overload or maintenance."
+          }
+        },
+        "deprecated":false
+      }
+    },
+    "/v2/directions/{profile}/geojson":{
+      "post":{
+        "tags":[
+          "Directions"
+        ],
+        "summary":"Directions Service GeoJSON (POST)",
+        "description":"Returns a route between two or more locations for a selected profile and its settings as GeoJSON",
+        "operationId":"getGeoJsonRouteUsingPOST",
+        "consumes":[
+          "application/json"
+        ],
+        "produces":[
+          "application/geo+json;charset=UTF-8",
+          "application/geo+json"
+        ],
+        "parameters":[
+          {
+            "name":"profile",
+            "in":"path",
+            "description":"Specifies the route profile.",
+            "required":true,
+            "type":"string",
+            "x-example":"driving-car",
+            "enum":[
+              "driving-car",
+              "driving-hgv",
+              "cycling-regular",
+              "cycling-road",
+              "cycling-mountain",
+              "cycling-electric",
+              "foot-walking",
+              "foot-hiking",
+              "wheelchair"
+            ]
+          },
+          {
+            "in":"body",
+            "name":"request",
+            "description":"The request payload",
+            "required":true,
+            "schema":{
+              "$ref":"#/definitions/Directions Service"
+            }
+          }
+        ],
+        "responses":{
+          "200":{
+            "description":"Standard response for successfully processed requests. Returns GeoJSON. The decoded values of the extra information can be found [here]([REDACTED_PROFILE_URL]/openrouteservice-docs).",
+            "schema":{
+              "$ref":"#/definitions/GeoJSONRouteResponse"
+            }
+          },
+          "201":{
+            "description":"Created"
+          },
+          "400":{
+            "description":"The request is incorrect and therefore can not be processed."
+          },
+          "401":{
+            "description":"Unauthorized"
+          },
+          "403":{
+            "description":"Forbidden"
+          },
+          "404":{
+            "description":"An element could not be found. If possible, a more detailed error code is provided."
+          },
+          "405":{
+            "description":"The specified HTTP method is not supported. For more details, refer to the EndPoint documentation."
+          },
+          "413":{
+            "description":"The request is larger than the server is able to process, the data provided in the request exceeds the capacity limit."
+          },
+          "500":{
+            "description":"An unexpected error was encountered and a more detailed error code is provided."
+          },
+          "501":{
+            "description":"Indicates that the server does not support the functionality needed to fulfill the request."
+          },
+          "503":{
+            "description":"The server is currently unavailable due to overload or maintenance."
+          }
+        },
+        "deprecated":false
+      }
+    },
+    "/v2/directions/{profile}/gpx":{
+      "post":{
+        "tags":[
+          "Directions"
+        ],
+        "summary":"Directions Service GPX (POST)",
+        "description":"Returns a route between two or more locations for a selected profile and its settings as GPX. The schema can be found [here](https://raw.githubusercontent.com/GIScience/openrouteservice-schema/master/gpx/v1/ors-gpx.xsd)",
+        "operationId":"getGPXRouteUsingPOST",
+        "consumes":[
+          "application/json"
+        ],
+        "produces":[
+          "application/gpx+xml;charset=UTF-8",
+          "application/gpx+xml"
+        ],
+        "parameters":[
+          {
+            "name":"profile",
+            "in":"path",
+            "description":"Specifies the route profile.",
+            "required":true,
+            "type":"string",
+            "x-example":"driving-car",
+            "enum":[
+              "driving-car",
+              "driving-hgv",
+              "cycling-regular",
+              "cycling-road",
+              "cycling-mountain",
+              "cycling-electric",
+              "foot-walking",
+              "foot-hiking",
+              "wheelchair"
+            ]
+          },
+          {
+            "in":"body",
+            "name":"request",
+            "description":"The request payload",
+            "required":true,
+            "schema":{
+              "$ref":"#/definitions/Directions Service"
+            }
+          }
+        ],
+        "responses":{
+          "200":{
+            "description":"Standard response for successfully processed requests. Returns GPX.",
+            "schema":{
+              "$ref":"#/definitions/gpx"
+            }
+          },
+          "201":{
+            "description":"Created"
+          },
+          "400":{
+            "description":"The request is incorrect and therefore can not be processed."
+          },
+          "401":{
+            "description":"Unauthorized"
+          },
+          "403":{
+            "description":"Forbidden"
+          },
+          "404":{
+            "description":"An element could not be found. If possible, a more detailed error code is provided."
+          },
+          "405":{
+            "description":"The specified HTTP method is not supported. For more details, refer to the EndPoint documentation."
+          },
+          "413":{
+            "description":"The request is larger than the server is able to process, the data provided in the request exceeds the capacity limit."
+          },
+          "500":{
+            "description":"An unexpected error was encountered and a more detailed error code is provided."
+          },
+          "501":{
+            "description":"Indicates that the server does not support the functionality needed to fulfill the request."
+          },
+          "503":{
+            "description":"The server is currently unavailable due to overload or maintenance."
+          }
+        },
+        "deprecated":false
+      }
+    },
+    "/v2/directions/{profile}/json":{
+      "post":{
+        "tags":[
+          "Directions"
+        ],
+        "summary":"Directions Service JSON (POST)",
+        "description":"Returns a route between two or more locations for a selected profile and its settings as JSON",
+        "operationId":"getJsonRouteUsingPOST",
+        "consumes":[
+          "application/json"
+        ],
+        "produces":[
+          "application/json;charset=UTF-8",
+          "application/json"
+        ],
+        "parameters":[
+          {
+            "name":"profile",
+            "in":"path",
+            "description":"Specifies the route profile.",
+            "required":true,
+            "type":"string",
+            "x-example":"driving-car",
+            "enum":[
+              "driving-car",
+              "driving-hgv",
+              "cycling-regular",
+              "cycling-road",
+              "cycling-mountain",
+              "cycling-electric",
+              "foot-walking",
+              "foot-hiking",
+              "wheelchair"
+            ]
+          },
+          {
+            "in":"body",
+            "name":"request",
+            "description":"The request payload",
+            "required":true,
+            "schema":{
+              "$ref":"#/definitions/Directions Service"
+            }
+          }
+        ],
+        "responses":{
+          "200":{
+            "description":"JSON Response",
+            "schema":{
+              "$ref":"#/definitions/JSONRouteResponse"
+            }
+          },
+          "201":{
+            "description":"Created"
+          },
+          "400":{
+            "description":"The request is incorrect and therefore can not be processed."
+          },
+          "401":{
+            "description":"Unauthorized"
+          },
+          "403":{
+            "description":"Forbidden"
+          },
+          "404":{
+            "description":"An element could not be found. If possible, a more detailed error code is provided."
+          },
+          "405":{
+            "description":"The specified HTTP method is not supported. For more details, refer to the EndPoint documentation."
+          },
+          "413":{
+            "description":"The request is larger than the server is able to process, the data provided in the request exceeds the capacity limit."
+          },
+          "500":{
+            "description":"An unexpected error was encountered and a more detailed error code is provided."
+          },
+          "501":{
+            "description":"Indicates that the server does not support the functionality needed to fulfill the request."
+          },
+          "503":{
+            "description":"The server is currently unavailable due to overload or maintenance."
+          }
+        },
+        "deprecated":false
+      }
+    },
+    "/v2/health":{
+      "get":{
+        "tags":[
+          "health-api"
+        ],
+        "summary":"fetchHealth",
+        "operationId":"fetchHealthUsingGET",
+        "produces":[
+          "*/*"
+        ],
+        "responses":{
+          "200":{
+            "description":"OK",
+            "schema":{
+              "type":"object"
+            }
+          },
+          "401":{
+            "description":"Unauthorized"
+          },
+          "403":{
+            "description":"Forbidden"
+          },
+          "404":{
+            "description":"Not Found"
+          }
+        },
+        "deprecated":false
+      }
+    },
+    "/v2/isochrones/{profile}":{
+      "post":{
+        "tags":[
+          "Isochrones"
+        ],
+        "summary":"Isochrones Service",
+        "description":"The Isochrone Service supports time and distance analyses for one single or multiple locations.\nYou may also specify the isochrone interval or provide multiple exact isochrone range values.\nThis service allows the same range of profile options as the /directions endpoint,\nwhich help you to further customize your request to obtain a more detailed reachability area response.",
+        "operationId":"getDefaultIsochronesUsingPOST",
+        "consumes":[
+          "application/json",
+          "application/geo+json"
+        ],
+        "produces":[
+          "application/geo+json;charset=UTF-8"
+        ],
+        "parameters":[
+          {
+            "name":"profile",
+            "in":"path",
+            "description":"Specifies the route profile.",
+            "required":true,
+            "type":"string",
+            "x-example":"driving-car",
+            "enum":[
+              "driving-car",
+              "driving-hgv",
+              "cycling-regular",
+              "cycling-road",
+              "cycling-mountain",
+              "cycling-electric",
+              "foot-walking",
+              "foot-hiking",
+              "wheelchair"
+            ]
+          },
+          {
+            "in":"body",
+            "name":"request",
+            "description":"The request payload",
+            "required":true,
+            "schema":{
+              "$ref":"#/definitions/IsochronesRequest"
+            }
+          }
+        ],
+        "responses":{
+          "200":{
+            "description":"Standard response for successfully processed requests. Returns GeoJSON.",
+            "schema":{
+              "$ref":"#/definitions/GeoJSONIsochronesResponse"
+            }
+          },
+          "201":{
+            "description":"Created"
+          },
+          "400":{
+            "description":"The request is incorrect and therefore can not be processed."
+          },
+          "401":{
+            "description":"Unauthorized"
+          },
+          "403":{
+            "description":"Forbidden"
+          },
+          "404":{
+            "description":"An element could not be found. If possible, a more detailed error code is provided."
+          },
+          "405":{
+            "description":"The specified HTTP method is not supported. For more details, refer to the EndPoint documentation."
+          },
+          "413":{
+            "description":"The request is larger than the server is able to process, the data provided in the request exceeds the capacity limit."
+          },
+          "500":{
+            "description":"An unexpected error was encountered and a more detailed error code is provided."
+          },
+          "501":{
+            "description":"Indicates that the server does not support the functionality needed to fulfill the request."
+          },
+          "503":{
+            "description":"The server is currently unavailable due to overload or maintenance."
+          }
+        },
+        "deprecated":false
+      }
+    },
+    "/v2/matrix/{profile}":{
+      "post":{
+        "tags":[
+          "Matrix"
+        ],
+        "summary":"Matrix Service",
+        "description":"Returns duration or distance matrix for mutliple source and destination points.\nBy default a symmetric duration matrix is returned where every point in locations is paired with each other. The result is null if a value can’t be determined.",
+        "operationId":"getDefaultUsingPOST",
+        "consumes":[
+          "application/json"
+        ],
+        "produces":[
+          "application/json;charset=UTF-8",
+          "application/json"
+        ],
+        "parameters":[
+          {
+            "name":"profile",
+            "in":"path",
+            "description":"Specifies the matrix profile.",
+            "required":true,
+            "type":"string",
+            "x-example":"driving-car",
+            "enum":[
+              "driving-car",
+              "driving-hgv",
+              "cycling-regular",
+              "cycling-road",
+              "cycling-mountain",
+              "cycling-electric",
+              "foot-walking",
+              "foot-hiking",
+              "wheelchair"
+            ]
+          },
+          {
+            "in":"body",
+            "name":"request",
+            "description":"The request payload",
+            "required":true,
+            "schema":{
+              "$ref":"#/definitions/MatrixRequest"
+            }
+          }
+        ],
+        "responses":{
+          "200":{
+            "description":"Standard response for successfully processed requests. Returns JSON.",
+            "schema":{
+              "$ref":"#/definitions/JSONMatrixResponse"
+            }
+          },
+          "201":{
+            "description":"Created"
+          },
+          "400":{
+            "description":"The request is incorrect and therefore can not be processed."
+          },
+          "401":{
+            "description":"Unauthorized"
+          },
+          "403":{
+            "description":"Forbidden"
+          },
+          "404":{
+            "description":"An element could not be found. If possible, a more detailed error code is provided."
+          },
+          "405":{
+            "description":"The specified HTTP method is not supported. For more details, refer to the EndPoint documentation."
+          },
+          "413":{
+            "description":"The request is larger than the server is able to process, the data provided in the request exceeds the capacity limit."
+          },
+          "500":{
+            "description":"An unexpected error was encountered and a more detailed error code is provided."
+          },
+          "501":{
+            "description":"Indicates that the server does not support the functionality needed to fulfill the request."
+          },
+          "503":{
+            "description":"The server is currently unavailable due to overload or maintenance."
+          }
+        },
+        "deprecated":false
+      }
+    },
+    "/v2/status":{
+      "get":{
+        "tags":[
+          "status-api"
+        ],
+        "summary":"fetchHealth",
+        "operationId":"fetchHealthUsingGET_1",
+        "produces":[
+          "*/*"
+        ],
+        "responses":{
+          "200":{
+            "description":"OK",
+            "schema":{
+              "$ref":"#/definitions/ResponseEntity"
+            }
+          },
+          "401":{
+            "description":"Unauthorized"
+          },
+          "403":{
+            "description":"Forbidden"
+          },
+          "404":{
+            "description":"Not Found"
+          }
+        },
+        "deprecated":false
+      }
+    }
+  },
+  "definitions":{
+    "Directions Service":{
+      "title":"Directions Service",
+      "allOf":[
+        {
+          "$ref":"#/definitions/Route Options"
+        },
+        {
+          "type":"object",
+          "required":[
+            "coordinates"
+          ],
+          "properties":{
+            "attributes":{
+              "type":"array",
+              "example":[
+                "avgspeed",
+                "percentage"
+              ],
+              "description":"List of route attributes",
+              "items":{
+                "type":"string",
+                "enum":[
+                  "avgspeed",
+                  "detourfactor",
+                  "percentage"
+                ]
+              }
+            },
+            "bearings":{
+              "type":"array",
+              "example":[
+                [
+                  30,
+                  20
+                ],
+                [
+
+                ],
+                [
+                  40,
+                  20
+                ]
+              ],
+              "description":"Specifies a list of pairs (bearings and deviations) to filter the segments of the road network a waypoint can snap to. For example `bearings=[[45,10],[120,20]]`. \nEach pair is a comma-separated list that can consist of one or two float values, where the first value is the bearing and the second one is the allowed deviation from the bearing. The bearing can take values between `0` and `360` clockwise from true north. If the deviation is not set, then the default value of `100` degrees is used. The number of pairs must correspond to the number of waypoints.\nThe number of bearings corresponds to the length of waypoints-1 or waypoints. If the bearing information for the last waypoint is given, then this will control the sector from which the destination waypoint may be reached. You can skip a bearing for a certain waypoint by passing an empty value for an array, e.g. `[30,20],[],[40,20]`. CUSTOM_KEYS:{'validWhen':{'ref':'profile','value':'cycling-*'}}",
+              "items":{
+                "type":"array",
+                "items":{
+                  "type":"number",
+                  "format":"double"
+                }
+              }
+            },
+            "continue_straight":{
+              "type":"boolean",
+              "example":false,
+              "description":"Forces the route to keep going straight at waypoints restricting uturns there even if it would be faster. This setting will work for all profiles except for `driving-*`. CUSTOM_KEYS:{'apiDefault':'true','validWhen':{'ref':'profile','valueNot':['driving-*']}}"
+            },
+            "coordinates":{
+              "type":"array",
+              "example":[
+                [
+                  8.681495,
+                  49.41461
+                ],
+                [
+                  8.686507,
+                  49.41943
+                ],
+                [
+                  8.687872,
+                  49.420318
+                ]
+              ],
+              "description":"The waypoints to use for the route as an array of `longitude/latitude` pairs",
+              "items":{
+                "type":"array",
+                "items":{
+                  "type":"number",
+                  "format":"double"
+                }
+              }
+            },
+            "elevation":{
+              "type":"boolean",
+              "example":false,
+              "description":"Specifies whether to return elevation values for points. Please note that elevation also gets encoded for json response encoded polyline."
+            },
+            "extra_info":{
+              "type":"array",
+              "example":[
+                "waytype",
+                "surface"
+              ],
+              "description":"The extra info items to include in the response",
+              "items":{
+                "type":"string",
+                "enum":[
+                  "steepness",
+                  "suitability",
+                  "surface",
+                  "waycategory",
+                  "waytype",
+                  "tollways",
+                  "traildifficulty",
+                  "osmid",
+                  "roadaccessrestrictions",
+                  "countryinfo"
+                ]
+              }
+            },
+            "geometry":{
+              "type":"boolean",
+              "example":true,
+              "description":"Specifies whether to return geometry. CUSTOM_KEYS:{'apiDefault':true, 'validWhen':{'ref':'format','value':['json']}}"
+            },
+            "geometry_simplify":{
+              "type":"boolean",
+              "example":false,
+              "description":"Specifies whether to simplify the geometry. Simplify geometry cannot be applied to routes with more than **one segment** and when `extra_info` is required.CUSTOM_KEYS:{'apiDefault':false}"
+            },
+            "id":{
+              "type":"string",
+              "example":"routing_request",
+              "description":"Arbitrary identification string of the request reflected in the meta information."
+            },
+            "instructions":{
+              "type":"boolean",
+              "example":true,
+              "description":"Specifies whether to return instructions. CUSTOM_KEYS:{'apiDefault':true}"
+            },
+            "instructions_format":{
+              "type":"string",
+              "example":"text",
+              "description":"Select html for more verbose instructions. CUSTOM_KEYS:{'apiDefault':'text'}",
+              "enum":[
+                "html",
+                "text"
+              ]
+            },
+            "language":{
+              "type":"string",
+              "example":"en",
+              "description":"Language for the route instructions. CUSTOM_KEYS:{'apiDefault':'en'}",
+              "enum":[
+                "de",
+                "en",
+                "en-us",
+                "es",
+                "fr",
+                "gr",
+                "hu",
+                "it",
+                "ne",
+                "nl",
+                "pt",
+                "ru",
+                "zh-cn"
+              ]
+            },
+            "maneuvers":{
+              "type":"boolean",
+              "example":false,
+              "description":"Specifies whether the maneuver object is included into the step object or not. CUSTOM_KEYS:{'apiDefault':false}"
+            },
+            "options":{
+              "example":{
+                "maximum_speed":100
+              },
+              "description":"For advanced options formatted as json object. For structure refer to the [these examples]([REDACTED_PROFILE_URL]/openrouteservice-docs#examples).",
+              "$ref":"#/definitions/Route Options"
+            },
+            "preference":{
+              "type":"string",
+              "example":"fastest",
+              "description":"Specifies the route preference. CUSTOM_KEYS:{'apiDefault':'fastest'}",
+              "enum":[
+                "fastest",
+                "shortest",
+                "recommended"
+              ]
+            },
+            "radiuses":{
+              "type":"array",
+              "example":[
+                200,
+                -1,
+                30
+              ],
+              "description":"A pipe list of maximum distances (measured in metres) that limit the search of nearby road segments to every given waypoint. The values must be greater than 0, the value of -1 specifies no limit in the search. The number of radiuses correspond to the number of waypoints.",
+              "items":{
+                "type":"number",
+                "format":"double"
+              }
+            },
+            "roundabout_exits":{
+              "type":"boolean",
+              "example":false,
+              "description":"Provides bearings of the entrance and all passed roundabout exits. Adds the `exit_bearings` array to the step object in the response. CUSTOM_KEYS:{'apiDefault':false}"
+            },
+            "skip_segments":{
+              "type":"array",
+              "example":[
+                2,
+                4
+              ],
+              "description":"Specifies the segments that should be skipped in the route calculation. A segment is the connection between two given coordinates and the counting starts with 1 for the connection between the first and second coordinate.",
+              "items":{
+                "type":"integer",
+                "format":"int32"
+              }
+            },
+            "suppress_warnings":{
+              "type":"boolean",
+              "example":false,
+              "description":"Suppress warning messages in the response"
+            },
+            "units":{
+              "type":"string",
+              "example":"m",
+              "description":"Specifies the distance unit. CUSTOM_KEYS:{'apiDefault':'m'}",
+              "enum":[
+                "m",
+                "km",
+                "mi"
+              ]
+            }
+          },
+          "title":"Directions Service",
+          "description":"The JSON body request sent to the routing service which defines options and parameters regarding the route to generate."
+        }
+      ],
+      "description":"The JSON body request sent to the routing service which defines options and parameters regarding the route to generate."
+    },
+    "EngineInfo":{
+      "type":"object",
+      "properties":{
+        "build_date":{
+          "type":"string",
+          "example":"2019-02-07T14:28:11Z",
+          "description":"The date that the service was last updated"
+        },
+        "version":{
+          "type":"string",
+          "example":5.0,
+          "description":"The backend version of the openrouteservice that was queried"
+        }
+      },
+      "title":"EngineInfo",
+      "description":"Information about the version of the openrouteservice that was used to generate the route"
+    },
+    "GeoJSONIsochroneBase":{
+      "type":"object",
+      "properties":{
+        "geometry":{
+          "type":"object",
+          "additionalProperties":{
+            "type":"object"
+          }
+        },
+        "type":{
+          "type":"string"
+        }
+      },
+      "title":"GeoJSONIsochroneBase"
+    },
+    "GeoJSONIsochronesResponse":{
+      "type":"object",
+      "properties":{
+        "bbox":{
+          "type":"array",
+          "example":[
+            49.414057,
+            8.680894,
+            49.420514,
+            8.690123
+          ],
+          "description":"Bounding box that covers all returned isochrones",
+          "items":{
+            "type":"number",
+            "format":"double"
+          }
+        },
+        "features":{
+          "type":"array",
+          "items":{
+            "$ref":"#/definitions/GeoJSONIsochroneBase"
+          }
+        },
+        "metadata":{
+          "$ref":"#/definitions/RouteResponseInfo"
+        },
+        "type":{
+          "type":"string"
+        }
+      },
+      "title":"GeoJSONIsochronesResponse"
+    },
+    "GeoJSONRouteResponse":{
+      "type":"object",
+      "properties":{
+        "bbox":{
+          "type":"array",
+          "example":[
+            49.414057,
+            8.680894,
+            49.420514,
+            8.690123
+          ],
+          "description":"Bounding box that covers all returned routes",
+          "items":{
+            "type":"number",
+            "format":"double"
+          }
+        },
+        "features":{
+          "type":"array",
+          "items":{
+            "type":"object"
+          }
+        },
+        "metadata":{
+          "$ref":"#/definitions/RouteResponseInfo"
+        },
+        "type":{
+          "type":"string"
+        }
+      },
+      "title":"GeoJSONRouteResponse"
+    },
+    "IsochronesRequest":{
+      "type":"object",
+      "required":[
+        "locations",
+        "range"
+      ],
+      "properties":{
+        "area_units":{
+          "type":"string",
+          "description":"Specifies the area unit.\nDefault: m. CUSTOM_KEYS:{'apiDefault':'m','validWhen':{'ref':'attributes','value':'area'}}",
+          "enum":[
+            "m",
+            "km",
+            "mi"
+          ]
+        },
+        "attributes":{
+          "type":"array",
+          "example":[
+            "area"
+          ],
+          "description":"List of isochrones attributes",
+          "items":{
+            "type":"string",
+            "enum":[
+              "area",
+              "reachfactor",
+              "total_pop"
+            ]
+          }
+        },
+        "id":{
+          "type":"string",
+          "example":"isochrones_request",
+          "description":"Arbitrary identification string of the request reflected in the meta information."
+        },
+        "intersections":{
+          "type":"boolean",
+          "description":"Specifies whether to return intersecting polygons. CUSTOM_KEYS:{'apiDefault':false}"
+        },
+        "interval":{
+          "type":"number",
+          "format":"double",
+          "example":30.0,
+          "description":"Interval of isochrones or equidistants for one range value. Value in **seconds** for time and **meters** for distance."
+        },
+        "location_type":{
+          "type":"string",
+          "example":"start",
+          "description":"`start` treats the location(s) as starting point, `destination` as goal. CUSTOM_KEYS:{'apiDefault':'start'}",
+          "enum":[
+            "start",
+            "destination"
+          ]
+        },
+        "locations":{
+          "type":"array",
+          "example":[
+            [
+              8.681495,
+              49.41461
+            ],
+            [
+              8.686507,
+              49.41943
+            ]
+          ],
+          "description":"The locations to use for the route as an array of `longitude/latitude` pairs",
+          "items":{
+            "type":"array",
+            "items":{
+              "type":"number",
+              "format":"double"
+            }
+          }
+        },
+        "options":{
+          "example":{
+            "avoid_borders":"all"
+          },
+          "description":"Additional options for the isochrones request",
+          "$ref":"#/definitions/Route Options"
+        },
+        "range":{
+          "type":"array",
+          "example":[
+            300,
+            200
+          ],
+          "description":"Maximum range value of the analysis in **seconds** for time and **metres** for distance.Alternatively a comma separated list of specific single range values if more than one location is set.",
+          "items":{
+            "type":"number",
+            "format":"double"
+          }
+        },
+        "range_type":{
+          "type":"string",
+          "example":"time",
+          "description":"Specifies the isochrones reachability type. CUSTOM_KEYS:{'apiDefault':'time'}",
+          "enum":[
+            "time",
+            "distance"
+          ]
+        },
+        "smoothing":{
+          "type":"number",
+          "format":"double",
+          "example":0.25,
+          "description":"Applies a level of generalisation to the isochrone polygons generated as a `smoothing_factor` between `0` and `1.0`.\nGeneralisation is produced by determining a maximum length of a connecting line between two points found on the outside of a containing polygon.\nIf the distance is larger than a threshold value, the line between the two points is removed and a smaller connecting line between other points is used.\nThe threshold value is determined as `(smoothing_factor * maximum_radius_of_isochrone) / 10`.\nTherefore, a value closer to 1 will result in a more generalised shape.\nThe polygon generation algorithm is based on Duckham and al. (2008) `\"Efficient generation of simple polygons for characterizing the shape of a set of points in the plane.\"`"
+        },
+        "units":{
+          "type":"string",
+          "example":"m",
+          "description":"Specifies the distance units only if `range_type` is set to distance.\nDefault: m. CUSTOM_KEYS:{'apiDefault':'m','validWhen':{'ref':'range_type','value':'distance'}}",
+          "enum":[
+            "m",
+            "km",
+            "mi"
+          ]
+        }
+      },
+      "title":"IsochronesRequest",
+      "description":"The JSON body request sent to the isochrones service which defines options and parameters regarding the isochrones to generate."
+    },
+    "JSON2DDestinations":{
+      "type":"object",
+      "properties":{
+        "location":{
+          "type":"array",
+          "example":[
+            8.678962,
+            49.40783
+          ],
+          "description":"{longitude},{latitude} coordinates of the closest accessible point on the routing graph",
+          "items":{
+            "type":"number",
+            "format":"double"
+          }
+        },
+        "name":{
+          "type":"string",
+          "example":"Bergheimer Straße",
+          "description":"Name of the street the closest accessible point is situated on. Only for `resolve_locations=true` and only if name is available.CUSTOM_KEYS:{'validWhen':{'ref':'resolve_locations','value':true}}"
+        },
+        "snapped_distance":{
+          "type":"number",
+          "format":"double",
+          "example":1.2,
+          "description":"Distance between the `source/destination` Location and the used point on the routing graph."
+        }
+      },
+      "title":"JSON2DDestinations"
+    },
+    "JSON2DSources":{
+      "type":"object",
+      "properties":{
+        "location":{
+          "type":"array",
+          "example":[
+            8.678962,
+            49.40783
+          ],
+          "description":"{longitude},{latitude} coordinates of the closest accessible point on the routing graph",
+          "items":{
+            "type":"number",
+            "format":"double"
+          }
+        },
+        "name":{
+          "type":"string",
+          "example":"Bergheimer Straße",
+          "description":"Name of the street the closest accessible point is situated on. Only for `resolve_locations=true` and only if name is available.CUSTOM_KEYS:{'validWhen':{'ref':'resolve_locations','value':true}}"
+        },
+        "snapped_distance":{
+          "type":"number",
+          "format":"double",
+          "example":1.2,
+          "description":"Distance between the `source/destination` Location and the used point on the routing graph."
+        }
+      },
+      "title":"JSON2DSources"
+    },
+    "JSONExtra":{
+      "type":"object",
+      "properties":{
+        "summary":{
+          "type":"array",
+          "description":"List representing the summary of the extra info items.",
+          "items":{
+            "$ref":"#/definitions/JSONExtraSummary"
+          }
+        },
+        "values":{
+          "type":"array",
+          "example":[
+            [
+              0,
+              3,
+              26
+            ],
+            [
+              3,
+              10,
+              12
+            ]
+          ],
+          "description":"A list of values representing a section of the route. The individual values are: \nValue 1: Indice of the staring point of the geometry for this section,\nValue 2: Indice of the end point of the geoemetry for this sections,\nValue 3: [Value]([REDACTED_PROFILE_URL]/openrouteservice-docs#routing-response) assigned to this section.",
+          "items":{
+            "type":"array",
+            "items":{
+              "type":"integer",
+              "format":"int64"
+            }
+          }
+        }
+      },
+      "title":"JSONExtra",
+      "description":"An object representing one of the extra info items requested"
+    },
+    "JSONExtraSummary":{
+      "type":"object",
+      "properties":{
+        "amount":{
+          "type":"number",
+          "format":"double",
+          "example":23.8,
+          "description":"Category percentage of the entire route."
+        },
+        "distance":{
+          "type":"number",
+          "format":"double",
+          "example":123.1,
+          "description":"Cumulative distance of this value."
+        },
+        "value":{
+          "type":"number",
+          "format":"double",
+          "example":5.0,
+          "description":"[Value]([REDACTED_PROFILE_URL]/openrouteservice-docs#routing-response) of a info category."
+        }
+      },
+      "title":"JSONExtraSummary"
+    },
+    "JSONIndividualRouteResponse":{
+      "type":"object",
+      "properties":{
+        "bbox":{
+          "type":"array",
+          "example":[
+            49.414057,
+            8.680894,
+            49.420514,
+            8.690123
+          ],
+          "description":"A bounding box which contains the entire route",
+          "items":{
+            "type":"number",
+            "format":"double"
+          }
+        },
+        "extras":{
+          "type":"object",
+          "description":"List of extra info objects representing the extra info items that were requested for the route.",
+          "additionalProperties":{
+            "$ref":"#/definitions/JSONExtra"
+          }
+        },
+        "segments":{
+          "type":"array",
+          "description":"List containing the segments and its corresponding steps which make up the route.",
+          "items":{
+            "$ref":"#/definitions/JSONSegment"
+          }
+        },
+        "summary":{
+          "description":"Summary information about the route",
+          "$ref":"#/definitions/JSONSummary"
+        },
+        "warnings":{
+          "type":"array",
+          "description":"List of warnings that have been generated for the route",
+          "items":{
+            "$ref":"#/definitions/JSONWarning"
+          }
+        },
+        "way_points":{
+          "type":"array",
+          "example":[
+            0,
+            23
+          ],
+          "description":"List containing the indices of way points corresponding to the *geometry*.",
+          "items":{
+            "type":"integer",
+            "format":"int32"
+          }
+        }
+      },
+      "title":"JSONIndividualRouteResponse",
+      "description":"An individual JSON based route created by the service"
+    },
+    "JSONMatrixResponse":{
+      "type":"object",
+      "properties":{
+        "destinations":{
+          "type":"array",
+          "description":"The individual destinations of the matrix calculations.",
+          "items":{
+            "$ref":"#/definitions/JSON2DDestinations"
+          }
+        },
+        "distances":{
+          "type":"array",
+          "example":[
+            [
+              0,
+              0.25
+            ],
+            [
+              0.25,
+              0
+            ]
+          ],
+          "description":"The distances of the matrix calculations.",
+          "items":{
+            "type":"array",
+            "items":{
+              "type":"number",
+              "format":"double"
+            }
+          }
+        },
+        "durations":{
+          "type":"array",
+          "example":[
+            [
+              0,
+              25
+            ],
+            [
+              25,
+              0
+            ]
+          ],
+          "description":"The durations of the matrix calculations.",
+          "items":{
+            "type":"array",
+            "items":{
+              "type":"number",
+              "format":"double"
+            }
+          }
+        },
+        "metadata":{
+          "description":"Information about the service and request",
+          "$ref":"#/definitions/MatrixResponseInfo"
+        },
+        "sources":{
+          "type":"array",
+          "description":"The individual sources of the matrix calculations.",
+          "items":{
+            "$ref":"#/definitions/JSON2DSources"
+          }
+        }
+      },
+      "title":"JSONMatrixResponse",
+      "description":"The Matrix Response contains one matrix for each specified `metrics` value."
+    },
+    "JSONRouteResponse":{
+      "type":"object",
+      "properties":{
+        "bbox":{
+          "type":"array",
+          "example":[
+            49.414057,
+            8.680894,
+            49.420514,
+            8.690123
+          ],
+          "description":"Bounding box that covers all returned routes",
+          "items":{
+            "type":"number",
+            "format":"double"
+          }
+        },
+        "metadata":{
+          "description":"Information about the service and request",
+          "$ref":"#/definitions/RouteResponseInfo"
+        },
+        "routes":{
+          "type":"array",
+          "description":"A list of routes returned from the request",
+          "items":{
+            "$ref":"#/definitions/JSONIndividualRouteResponse"
+          }
+        }
+      },
+      "title":"JSONRouteResponse"
+    },
+    "JSONSegment":{
+      "type":"object",
+      "properties":{
+        "ascent":{
+          "type":"number",
+          "format":"double",
+          "example":56.3,
+          "description":" Contains ascent of this segment in metres. CUSTOM_KEYS:{'validWhen':{'ref':'elevation',value:true}}"
+        },
+        "avgspeed":{
+          "type":"number",
+          "format":"double",
+          "example":56.3,
+          "description":"Contains the average speed of this segment in km/h. CUSTOM_KEYS:{'validWhen':{'ref':'attributes','valueContains':'avgspeed'}}"
+        },
+        "descent":{
+          "type":"number",
+          "format":"double",
+          "example":45.2,
+          "description":"Contains descent of this segment in metres. CUSTOM_KEYS:{'validWhen':{'ref':'elevation',value:true}}"
+        },
+        "detourfactor":{
+          "type":"number",
+          "format":"double",
+          "example":0.5,
+          "description":"Contains the deviation compared to a straight line that would have the factor `1`. Double the Distance would be a `2`. CUSTOM_KEYS:{'validWhen':{'ref':'attributes','valueContains':'detourfactor'}}"
+        },
+        "distance":{
+          "type":"number",
+          "format":"double",
+          "example":253.0,
+          "description":"Contains the distance of the segment in specified units."
+        },
+        "duration":{
+          "type":"number",
+          "format":"double",
+          "example":37.7,
+          "description":"Contains the duration of the segment in seconds."
+        },
+        "percentage":{
+          "type":"number",
+          "format":"double",
+          "example":43.2,
+          "description":"Contains the proportion of the route in percent. CUSTOM_KEYS:{'validWhen':{'ref':'attributes','valueContains':'percentage'}}"
+        },
+        "steps":{
+          "type":"array",
+          "description":"List containing the specific steps the segment consists of.",
+          "items":{
+            "$ref":"#/definitions/JSONStep"
+          }
+        }
+      },
+      "title":"JSONSegment",
+      "description":"List containing the segments and its correspoding steps which make up the route."
+    },
+    "JSONStep":{
+      "type":"object",
+      "properties":{
+        "distance":{
+          "type":"number",
+          "format":"double",
+          "example":245.0,
+          "description":"The distance for the step in metres."
+        },
+        "duration":{
+          "type":"number",
+          "format":"double",
+          "example":96.2,
+          "description":"The duration for the step in seconds."
+        },
+        "exit_bearings":{
+          "type":"array",
+          "example":[
+            10,
+            45,
+            60
+          ],
+          "description":"Contains the bearing of the entrance and all passed exits in a roundabout  CUSTOM_KEYS:{'validWhen':{'ref':'roundabout_exits',value:true}}.",
+          "items":{
+            "type":"integer",
+            "format":"int32"
+          }
+        },
+        "exit_number":{
+          "type":"integer",
+          "format":"int32",
+          "example":2,
+          "description":"Only for roundabouts. Contains the number of the exit to take."
+        },
+        "instruction":{
+          "type":"string",
+          "example":"Turn right onto Berliner Straße",
+          "description":"The routing instruction text for the step."
+        },
+        "maneuver":{
+          "description":"The maneuver to be performed  CUSTOM_KEYS:{'validWhen':{'ref':'maneuvers',value:true}}",
+          "$ref":"#/definitions/JSONStepManeuver"
+        },
+        "name":{
+          "type":"string",
+          "example":"Berliner Straße",
+          "description":"The name of the next street."
+        },
+        "type":{
+          "type":"integer",
+          "format":"int32",
+          "example":1,
+          "description":"The [instruction]([REDACTED_PROFILE_URL]/openrouteservice-docs#instruction-types) action for symbolisation purposes."
+        },
+        "way_points":{
+          "type":"array",
+          "example":[
+            45,
+            48
+          ],
+          "description":"List containing the indices of the steps start- and endpoint corresponding to the *geometry*.",
+          "items":{
+            "type":"integer",
+            "format":"int32"
+          }
+        }
+      },
+      "title":"JSONStep",
+      "description":"Step of a route segment"
+    },
+    "JSONStepManeuver":{
+      "type":"object",
+      "properties":{
+        "bearing_after":{
+          "type":"integer",
+          "format":"int32",
+          "example":96,
+          "description":"The azimuth angle (in degrees) of the direction right after the maneuver."
+        },
+        "bearing_before":{
+          "type":"integer",
+          "format":"int32",
+          "example":24,
+          "description":"The azimuth angle (in degrees) of the direction right before the maneuver."
+        },
+        "location":{
+          "type":"array",
+          "example":[
+            8.678962,
+            49.407819
+          ],
+          "description":"The coordinate of the point where a maneuver takes place.",
+          "items":{
+            "type":"number",
+            "format":"double"
+          }
+        }
+      },
+      "title":"JSONStepManeuver",
+      "description":"Maneuver object of the step"
+    },
+    "JSONSummary":{
+      "type":"object",
+      "properties":{
+        "ascent":{
+          "type":"number",
+          "format":"double",
+          "example":166.3,
+          "description":"Total ascent in meters.CUSTOM_KEYS:{'validWhen':{'ref':'elevation','value':true}}"
+        },
+        "descent":{
+          "type":"number",
+          "format":"double",
+          "example":201.3,
+          "description":"Total descent in meters.CUSTOM_KEYS:{'validWhen':{'ref':'elevation','value':true}}"
+        },
+        "distance":{
+          "type":"number",
+          "format":"double",
+          "example":12.6,
+          "description":"Total route distance in specified units."
+        },
+        "duration":{
+          "type":"number",
+          "format":"double",
+          "example":604.0,
+          "description":"Total duration in seconds."
+        }
+      },
+      "title":"JSONSummary",
+      "description":"Contains total sums of duration, route distance and actual distance of the route."
+    },
+    "JSONWarning":{
+      "type":"object",
+      "properties":{
+        "code":{
+          "type":"integer",
+          "format":"int32",
+          "example":1,
+          "description":"Identification code for the warning"
+        },
+        "message":{
+          "type":"string",
+          "example":"This route may go over restricted roads",
+          "description":"The message associated with the warning"
+        }
+      },
+      "title":"JSONWarning",
+      "description":"Informs about possible difficulties like access restrictions on the generated route. Generates a corresponding `extras` object with the affected segments."
+    },
+    "MatrixRequest":{
+      "type":"object",
+      "required":[
+        "locations"
+      ],
+      "properties":{
+        "destinations":{
+          "type":"array",
+          "description":"A list of indices that refers to the list of locations (starting with `0`). `{index_1},{index_2}[,{index_N} ...]` or `all` (default). `[0,3]` for the first and fourth locations CUSTOM_KEYS:{'apiDefault':['all']}",
+          "items":{
+            "type":"string"
+          }
+        },
+        "id":{
+          "type":"string",
+          "example":"matrix_request",
+          "description":"Arbitrary identification string of the request reflected in the meta information."
+        },
+        "locations":{
+          "type":"array",
+          "example":[
+            [
+              9.70093,
+              48.477473
+            ],
+            [
+              9.207916,
+              49.153868
+            ],
+            [
+              37.573242,
+              55.801281
+            ],
+            [
+              115.663757,
+              38.106467
+            ]
+          ],
+          "description":"List of comma separated lists of `longitude,latitude` coordinates.",
+          "items":{
+            "type":"array",
+            "items":{
+              "type":"number",
+              "format":"double"
+            }
+          }
+        },
+        "metrics":{
+          "type":"array",
+          "description":"Specifies a list of returned metrics.\n* `distance` - Returns distance matrix for specified points in defined `units`.\n* `duration` - Returns duration matrix for specified points in **seconds**. CUSTOM_KEYS:{'apiDefault':'duration'}",
+          "items":{
+            "type":"string",
+            "enum":[
+              "distance",
+              "duration"
+            ]
+          }
+        },
+        "resolve_locations":{
+          "type":"boolean",
+          "description":"Specifies whether given locations are resolved or not. If the parameter value set to `true`, every element in `destinations` and `sources` will contain a `name` element that identifies the name of the closest street. Default is `false`. CUSTOM_KEYS:{'apiDefault':false}"
+        },
+        "sources":{
+          "type":"array",
+          "description":"A list of indices that refers to the list of locations (starting with `0`). `{index_1},{index_2}[,{index_N} ...]` or `all` (default). example `[0,3]` for the first and fourth locations CUSTOM_KEYS:{'apiDefault':['all']}",
+          "items":{
+            "type":"string"
+          }
+        },
+        "units":{
+          "type":"string",
+          "description":"Specifies the distance unit.\nDefault: m. CUSTOM_KEYS:{'apiDefault':'m','validWhen':{'ref':'metrics','value':'distance'}`}",
+          "enum":[
+            "m",
+            "km",
+            "mi"
+          ]
+        }
+      },
+      "title":"MatrixRequest",
+      "description":"The JSON body request sent to the matrix service which defines options and parameters regarding the matrix to generate."
+    },
+    "MatrixResponseInfo":{
+      "type":"object",
+      "properties":{
+        "attribution":{
+          "type":"string",
+          "example":"openrouteservice.org, OpenStreetMap contributors",
+          "description":"Copyright and attribution information"
+        },
+        "engine":{
+          "description":"Information about the routing service",
+          "$ref":"#/definitions/EngineInfo"
+        },
+        "id":{
+          "type":"string",
+          "example":"request123",
+          "description":"ID of the request (as passed in by the query)"
+        },
+        "osm_file_md5_hash":{
+          "type":"string",
+          "example":"c0327ba6",
+          "description":"The MD5 hash of the OSM planet file that was used for generating graphs"
+        },
+        "query":{
+          "description":"The information that was used for generating the matrix",
+          "$ref":"#/definitions/MatrixRequest"
+        },
+        "service":{
+          "type":"string",
+          "example":"matrix",
+          "description":"The service that was requested"
+        },
+        "timestamp":{
+          "type":"integer",
+          "format":"int64",
+          "example":1549549847974,
+          "description":"Time that the request was made (UNIX Epoch time)"
+        }
+      },
+      "title":"MatrixResponseInfo",
+      "description":"Information about the request"
+    },
+    "Profile Parameters":{
+      "type":"object",
+      "properties":{
+        "restrictions":{
+          "$ref":"#/definitions/Restrictions"
+        },
+        "weightings":{
+          "$ref":"#/definitions/Profile Weightings"
+        }
+      },
+      "title":"Profile Parameters",
+      "description":"Specifies additional routing parameters. For all profiles except `driving-car`."
+    },
+    "Profile Weightings":{
+      "type":"object",
+      "properties":{
+        "green":{
+          "type":"number",
+          "format":"float",
+          "example":0.4,
+          "description":"Specifies the Green factor for `foot-*` profiles.\n\nfactor: Values range from 0 to 1. 0 equals normal routing. 1 will prefer ways through green areas over a shorter route. CUSTOM_KEYS:{'validWhen':{'ref':'profile','value':['foot-*']}}"
+        },
+        "quiet":{
+          "type":"number",
+          "format":"float",
+          "example":0.8,
+          "description":"Specifies the Quiet factor for foot-* profiles.\n\nfactor: Values range from 0 to 1. 0 equals normal routing. 1 will prefer quiet ways over a shorter route. CUSTOM_KEYS:{'validWhen':{'ref':'profile','value':['foot-*']}}"
+        },
+        "steepness_difficulty":{
+          "type":"integer",
+          "format":"int32",
+          "example":2,
+          "description":"Specifies the fitness level for `cycling-*` profiles.\n\n level: 0 = Novice, 1 = Moderate, 2 = Amateur, 3 = Pro. The prefered gradient increases with level. CUSTOM_KEYS:{'validWhen':{'ref':'profile','value':['cycling-*']}}"
+        }
+      },
+      "title":"Profile Weightings",
+      "description":"Describe additional weightings to be applied to edges on the routing."
+    },
+    "ResponseEntity":{
+      "type":"object",
+      "properties":{
+        "body":{
+          "type":"object"
+        },
+        "statusCode":{
+          "type":"string",
+          "enum":[
+            "100",
+            "101",
+            "102",
+            "103",
+            "200",
+            "201",
+            "202",
+            "203",
+            "204",
+            "205",
+            "206",
+            "207",
+            "208",
+            "226",
+            "300",
+            "301",
+            "302",
+            "303",
+            "304",
+            "305",
+            "307",
+            "308",
+            "400",
+            "401",
+            "402",
+            "403",
+            "404",
+            "405",
+            "406",
+            "407",
+            "408",
+            "409",
+            "410",
+            "411",
+            "412",
+            "413",
+            "414",
+            "415",
+            "416",
+            "417",
+            "418",
+            "419",
+            "420",
+            "421",
+            "422",
+            "423",
+            "424",
+            "426",
+            "428",
+            "429",
+            "431",
+            "451",
+            "500",
+            "501",
+            "502",
+            "503",
+            "504",
+            "505",
+            "506",
+            "507",
+            "508",
+            "509",
+            "510",
+            "511"
+          ]
+        },
+        "statusCodeValue":{
+          "type":"integer",
+          "format":"int32"
+        }
+      },
+      "title":"ResponseEntity"
+    },
+    "Restrictions":{
+      "type":"object",
+      "properties":{
+        "axleload":{
+          "type":"number",
+          "format":"float",
+          "example":50.0,
+          "description":"Axleload restriction in tons. CUSTOM_KEYS:{'validWhen':{'ref':'profile','value':['driving-hgv']}}"
+        },
+        "hazmat":{
+          "type":"boolean",
+          "description":"Specifies whether to use appropriate routing for delivering hazardous goods and avoiding water protected areas. Default is `false`. CUSTOM_KEYS:{'apiDefault':false,'validWhen':{'ref':'profile','value':['driving-hgv']}}"
+        },
+        "height":{
+          "type":"number",
+          "format":"float",
+          "example":4.2,
+          "description":"Height restriction in metres. CUSTOM_KEYS:{'validWhen':{'ref':'profile','value':['driving-hgv']}}"
+        },
+        "length":{
+          "type":"number",
+          "format":"float",
+          "example":8.4,
+          "description":"Length restriction in metres. CUSTOM_KEYS:{'validWhen':{'ref':'profile','value':['driving-hgv']}}"
+        },
+        "maximum_incline":{
+          "type":"integer",
+          "format":"int32",
+          "example":3,
+          "description":"Specifies the maximum incline as a percentage. `3`, `6` (default), `10`, `15.CUSTOM_KEYS:{'apiDefault':6,'validWhen':{'ref':'profile','value':['wheelchair']}}"
+        },
+        "maximum_sloped_kerb":{
+          "type":"number",
+          "format":"float",
+          "example":0.03,
+          "description":"Specifies the maximum height of the sloped curb in metres. Values are `0.03`, `0.06` (default), `0.1`.CUSTOM_KEYS:{'apiDefault':0.6,'validWhen':{'ref':'profile','value':['wheelchair']}}"
+        },
+        "minimum_width":{
+          "type":"number",
+          "format":"float",
+          "example":2.5,
+          "description":"Specifies the minimum width of the footway in metres.CUSTOM_KEYS:{'validWhen':{'ref':'profile','value':['wheelchair']}}"
+        },
+        "smoothness_type":{
+          "type":"string",
+          "example":"best",
+          "description":"Specifies the minimum smoothness of the route. Default is `good`.CUSTOM_KEYS:{'apiDefault':'good','validWhen':{'ref':'profile','value':['wheelchair']}}"
+        },
+        "surface_type":{
+          "type":"string",
+          "example":"asphalt",
+          "description":"Specifies the minimum surface type. Default is `cobblestone:flattened`. CUSTOM_KEYS:{'apiDefault':'cobblestone:flattened','validWhen':{'ref':'profile','value':['wheelchair']}}"
+        },
+        "track_type":{
+          "type":"string",
+          "example":"grade2",
+          "description":"Specifies the minimum grade of the route. Default is `grade1`. CUSTOM_KEYS:{'apiDefault':'grade1','validWhen':{'ref':'profile','value':['wheelchair']}}"
+        },
+        "weight":{
+          "type":"number",
+          "format":"float",
+          "example":40.0,
+          "description":"Weight restriction in tons. CUSTOM_KEYS:{'validWhen':{'ref':'profile','value':['driving-hgv']}}"
+        },
+        "width":{
+          "type":"number",
+          "format":"float",
+          "example":5.6,
+          "description":"Width restriction in metres. CUSTOM_KEYS:{'validWhen':{'ref':'profile','value':['driving-hgv']}}"
+        }
+      },
+      "title":"Restrictions",
+      "description":"Describe restrictions to be applied to edges on the routing. any edges that do not match these restrictions are not traversed."
+    },
+    "Route Options":{
+      "type":"object",
+      "properties":{
+        "avoid_borders":{
+          "type":"string",
+          "description":"`all` for no border crossing. `controlled` to cross open borders but avoid controlled ones. Only for `driving-*` profiles. CUSTOM_KEYS:{'validWhen':{'ref':'profile','value':['driving-*']}}",
+          "enum":[
+            "all",
+            "controlled",
+            "none"
+          ]
+        },
+        "avoid_countries":{
+          "type":"array",
+          "example":[
+            11,
+            193
+          ],
+          "description":"List of countries to exclude from routing with `driving-*` profiles. Can be used together with `'avoid_borders': 'controlled'`. `[ 11, 193 ]` would exclude Austria and Switzerland. List of countries and application examples can be found [here]([REDACTED_PROFILE_URL]/openrouteservice-docs#country-list). Also, ISO standard country codes cna be used in place of the numerical ids, for example, DE or DEU for Germany. CUSTOM_KEYS:{'validWhen':{'ref':'profile','value':['driving-*']}}",
+          "items":{
+            "type":"string"
+          }
+        },
+        "avoid_features":{
+          "type":"array",
+          "description":"List of features to avoid. CUSTOM_KEYS:{'itemRestrictions':{'ref':'profile', 'itemsWhen':{'driving-*':['highways','tollways','ferries'],'cycling-*':['ferries','steps','fords'],'foot-*':['ferries','fords','steps'],'wheelchair':['ferries','steps']}}}",
+          "items":{
+            "type":"string",
+            "enum":[
+              "highways",
+              "tollways",
+              "ferries",
+              "fords",
+              "steps"
+            ]
+          }
+        },
+        "avoid_polygons":{
+          "type":"object",
+          "description":"Comprises areas to be avoided for the route. Formatted in GeoJSON as either a Polygon or Multipolygon object.",
+          "additionalProperties":{
+            "type":"object"
+          }
+        },
+        "profile_params":{
+          "description":" Specifies additional routing parameters.",
+          "$ref":"#/definitions/Profile Parameters"
+        },
+        "vehicle_type":{
+          "type":"string",
+          "description":"(for profile=driving-hgv only): hgv,bus,agricultural,delivery,forestry and goods. It is needed for vehicle restrictions to work. CUSTOM_KEYS:{'validWhen':{'ref':'profile','value':['driving-hgv']}}",
+          "enum":[
+            "hgv",
+            "bus",
+            "agricultural",
+            "delivery",
+            "forestry",
+            "goods",
+            "unknown"
+          ]
+        }
+      },
+      "title":"Route Options",
+      "description":"Advanced options for routing"
+    },
+    "RouteResponseInfo":{
+      "type":"object",
+      "properties":{
+        "attribution":{
+          "type":"string",
+          "example":"openrouteservice.org | OpenStreetMap contributors",
+          "description":"Copyright and attribution information"
+        },
+        "engine":{
+          "description":"Information about the routing service",
+          "$ref":"#/definitions/EngineInfo"
+        },
+        "id":{
+          "type":"string",
+          "example":"request123",
+          "description":"ID of the request (as passed in by the query)"
+        },
+        "osm_file_md5_hash":{
+          "type":"string",
+          "example":"c0327ba6",
+          "description":"The MD5 hash of the OSM planet file that was used for generating graphs"
+        },
+        "query":{
+          "description":"The information that was used for generating the route",
+          "$ref":"#/definitions/Directions Service"
+        },
+        "service":{
+          "type":"string",
+          "example":"routing",
+          "description":"The service that was requested"
+        },
+        "timestamp":{
+          "type":"integer",
+          "format":"int64",
+          "example":1549549847974,
+          "description":"Time that the request was made (UNIX Epoch time)"
+        }
+      },
+      "title":"RouteResponseInfo",
+      "description":"Information about the request"
+    },
+    "gpx":{
+      "type":"object",
+      "properties":{
+        "gpxRouteElements":{
+          "type":"array",
+          "items":{
+            "$ref":"#/definitions/rte"
+          }
+        }
+      },
+      "title":"gpx",
+      "xml":{
+        "name":"gpx",
+        "attribute":false,
+        "wrapped":false
+      }
+    },
+    "rte":{
+      "type":"object",
+      "title":"rte",
+      "xml":{
+        "name":"rte",
+        "attribute":false,
+        "wrapped":false
+      }
+    }
+  }
+}
diff --git a/Engine/Plugins/Web/WebAPI/Source/WebAPIOpenAPI/Private/Tests/Samples/openstreetmap.json b/Engine/Plugins/Web/WebAPI/Source/WebAPIOpenAPI/Private/Tests/Samples/openstreetmap.json
new file mode 100644
--- /dev/null
+++ b/Engine/Plugins/Web/WebAPI/Source/WebAPIOpenAPI/Private/Tests/Samples/openstreetmap.json
@@ -0,0 +1,695 @@
+{
+    "openapi": "3.0.0",
+    "info": {
+      "title": "OpenStreetMap",
+      "description": "This is the editing API for OpenStreetMap.",
+      "version": "0.6",
+      "termsOfService": "https://operations.osmfoundation.org/policies/api/",
+      "license": {
+        "name": "ISC License",
+        "url": "[REDACTED_PROFILE_URL]/GoMap/blob/master/LICENSE.md"
+      }
+    },
+    "servers": [
+      {
+        "url": "https://api.openstreetmap.org",
+        "description": "Production server (uses live data)"
+      },
+      {
+        "url": "https://master.apis.dev.openstreetmap.org",
+        "description": "Sandbox server (uses test data)"
+      }
+    ],
+    "tags": [
+      {
+        "name": "map",
+        "description": "Operations related to map data"
+      },
+      {
+        "name": "user",
+        "description": "Everything on the user"
+      },
+      {
+        "name": "misc",
+        "description": "Meta-information on the API and the server"
+      }
+    ],
+    "paths": {
+      "/api/versions": {
+        "get": {
+          "tags": [
+            "misc"
+          ],
+          "summary": "Returns a list of API versions supported by this instance.",
+          "operationId": "getAPIVersions",
+          "responses": {
+            "200": {
+              "description": "successful operation",
+              "content": {
+                "application/xml": {
+                  "schema": {
+                    "$ref": "#/components/schemas/APIVersionsResponse"
+                  }
+                }
+              }
+            }
+          }
+        }
+      },
+      "/api/capabilities": {
+        "get": {
+          "tags": [
+            "misc"
+          ],
+          "summary": "This API call is meant to provide information about the capabilities and limitations of the current API.",
+          "operationId": "getAPICapabilities",
+          "responses": {
+            "200": {
+              "description": "successful operation",
+              "content": {
+                "application/xml": {
+                  "schema": {
+                    "$ref": "#/components/schemas/APICapabilitiesResponse"
+                  }
+                }
+              }
+            }
+          }
+        }
+      },
+      "/api/0.6/map": {
+        "get": {
+          "tags": [
+            "map"
+          ],
+          "summary": "Retrieves map data by the given bounding box.",
+          "description": "The operation returns:\n- All nodes that are inside a given bounding box and any relations that reference them.\n- All ways that reference at least one node that is inside a given bounding box, any relations that reference them [the ways],\n  and any nodes outside the bounding box that the ways may reference.\n\n- All relations that reference one of the nodes, ways or relations included due to the above rules.\n  (Does not apply recursively, see explanation below.)\n\n\nNote that, while this operation returns those relations that reference the aforementioned nodes and ways, the reverse is not true: it does not (necessarily) return all of the nodes and ways that are referenced by these relations. This prevents unreasonably-large result sets. For example, imagine the case where:\n- There is a relation named \"England\" that references every node in England.\n- The nodes, ways, and relations are retrieved for a bounding box that covers a small portion of England.\n\nWhile the result would include the nodes, ways, and relations as specified by the rules for the command,  including the \"England\" relation, it would (fortuitously) not include every node and way in England.  If desired, the nodes and ways referenced by the \"England\" relation could be retrieved by their respective IDs.\n\nAlso note that ways which intersect the bounding box but have no nodes within the bounding box will not be returned. \n",
+          "operationId": "getMapDataByBoundingBox",
+          "parameters": [
+            {
+              "in": "query",
+              "name": "bbox",
+              "schema": {
+                "type": "string",
+                "example": "13.377639198,52.5162399276,13.3778497142,52.5163460219"
+              },
+              "description": "The comma-separated bounding box for which to download the data.\n\n`left`,`bottom`,`right`,`top`, where:\n- `left` is the longitude of the left (westernmost) side of the bounding box.\n- `bottom` is the latitude of the bottom (southernmost) side of the bounding box.\n- `right` is the longitude of the right (easternmost) side of the bounding box.\n- `top` is the latitude of the top (northernmost) side of the bounding box.\n"
+            }
+          ],
+          "responses": {
+            "200": {
+              "content": {
+                "application/json": {
+                  "schema": {
+                    "type": "object",
+                    "allOf": [
+                      {
+                        "$ref": "#/components/schemas/OSMResponse"
+                      }
+                    ],
+                    "properties": {
+                      "bounds": {
+                        "type": "object",
+                        "properties": {
+                          "minlat": {
+                            "$ref": "#/components/schemas/Node/allOf/1/properties/lat"
+                          },
+                          "minlon": {
+                            "$ref": "#/components/schemas/Node/allOf/1/properties/lon"
+                          },
+                          "maxlat": {
+                            "$ref": "#/components/schemas/Node/allOf/1/properties/lat"
+                          },
+                          "maxlon": {
+                            "$ref": "#/components/schemas/Node/allOf/1/properties/lon"
+                          }
+                        }
+                      },
+                      "elements": {
+                        "type": "array",
+                        "items": {
+                          "oneOf": [
+                            {
+                              "$ref": "#/components/schemas/Node"
+                            },
+                            {
+                              "$ref": "#/components/schemas/Relation"
+                            },
+                            {
+                              "$ref": "#/components/schemas/Way"
+                            }
+                          ]
+                        }
+                      }
+                    },
+                    "xml": {
+                      "wrapped": true,
+                      "name": "osm"
+                    }
+                  }
+                }
+              }
+            },
+            "400": {
+              "description": "When any of the node/way/relation limits are exceeded, in particular if the call would return more than 50'000 nodes."
+            },
+            "509": {
+              "description": "\"Error: You have downloaded too much data. Please try again later.\" See [Developer FAQ](https://wiki.openstreetmap.org/wiki/Developer_FAQ#I.27ve_been_blocked_from_the_API_for_downloading_too_much._Now_what.3F).\n"
+            }
+          }
+        }
+      },
+      "/api/0.6/permissions": {
+        "get": {
+          "tags": [
+            "misc"
+          ],
+          "summary": "Returns the permissions granted to the current API connection.",
+          "operationId": "getAPIPermissions",
+          "security": [
+            {
+              "basicAuth": [
+  
+              ]
+            }
+          ],
+          "responses": {
+            "200": {
+              "description": "successful operation",
+              "content": {
+                "application/xml": {
+                  "schema": {
+                    "type": "object",
+                    "allOf": [
+                      {
+                        "$ref": "#/components/schemas/OSMResponse"
+                      }
+                    ],
+                    "properties": {
+                      "permissions": {
+                        "type": "array",
+                        "items": {
+                          "type": "object",
+                          "properties": {
+                            "name": {
+                              "$ref": "#/components/schemas/APIPermission"
+                            }
+                          },
+                          "xml": {
+                            "name": "permission"
+                          }
+                        },
+                        "xml": {
+                          "wrapped": true
+                        }
+                      }
+                    },
+                    "xml": {
+                      "wrapped": true,
+                      "name": "osm"
+                    }
+                  }
+                }
+              }
+            }
+          }
+        }
+      },
+      "/api/0.6/user/details.json": {
+        "get": {
+          "tags": [
+            "user"
+          ],
+          "summary": "Get the user's profile as well as their OSM messages and preferences",
+          "description": "The messages section has been available since mid-2013.",
+          "operationId": "getSignedInUserDetails",
+          "security": [
+            {
+              "basicAuth": [
+  
+              ]
+            }
+          ],
+          "responses": {
+            "200": {
+              "description": "successful operation",
+              "content": {
+                "application/json": {
+                  "schema": {
+                    "$ref": "#/components/schemas/UserDetailsResponse"
+                  }
+                }
+              }
+            }
+          }
+        }
+      }
+    },
+    "components": {
+      "securitySchemes": {
+        "basicAuth": {
+          "type": "http",
+          "scheme": "basic"
+        }
+      },
+      "schemas": {
+        "APICapabilities": {
+          "type": "object",
+          "properties": {
+            "version": {
+              "type": "object",
+              "description": "The API call versions that the server will accept",
+              "properties": {
+                "minimum": {
+                  "type": "number",
+                  "format": "double",
+                  "example": 0.6,
+                  "xml": {
+                    "attribute": true
+                  }
+                },
+                "maximum": {
+                  "type": "number",
+                  "format": "double",
+                  "example": 0.6,
+                  "xml": {
+                    "attribute": true
+                  }
+                }
+              }
+            },
+            "area": {
+              "type": "object",
+              "properties": {
+                "maximum": {
+                  "type": "number",
+                  "format": "double",
+                  "description": "The maximum area in square degrees that can be queried by API calls",
+                  "example": 0.25,
+                  "xml": {
+                    "attribute": true
+                  }
+                }
+              }
+            },
+            "note_area": {
+              "type": "object",
+              "properties": {
+                "maximum": {
+                  "type": "number",
+                  "example": 25,
+                  "xml": {
+                    "attribute": true
+                  }
+                }
+              }
+            },
+            "tracepoints": {
+              "type": "object",
+              "properties": {
+                "per_page": {
+                  "type": "number",
+                  "description": "The maximum number of points in a single GPS trace. (Possibly incorrect)",
+                  "example": 5000,
+                  "xml": {
+                    "attribute": true
+                  }
+                }
+              }
+            },
+            "waynodes": {
+              "type": "object",
+              "properties": {
+                "maximum": {
+                  "type": "number",
+                  "description": "The maximum number of nodes that a way may contain",
+                  "example": 2000,
+                  "xml": {
+                    "attribute": true
+                  }
+                }
+              }
+            },
+            "changesets": {
+              "type": "object",
+              "properties": {
+                "maximum_elements": {
+                  "type": "number",
+                  "description": "The maximum number of combined nodes, ways and relations that can be contained in a changeset",
+                  "example": 10000,
+                  "xml": {
+                    "attribute": true
+                  }
+                }
+              }
+            },
+            "timeout": {
+              "type": "object",
+              "properties": {
+                "seconds": {
+                  "type": "number",
+                  "example": 300,
+                  "xml": {
+                    "attribute": true
+                  }
+                }
+              }
+            },
+            "status": {
+              "type": "object",
+              "description": "The status of the database, API and GPX API. The database field is informational, and the API/GPX-API fields indicate what the client should expect.",
+              "properties": {
+                "database": {
+                  "$ref": "#/components/schemas/APIStatus"
+                },
+                "api": {
+                  "$ref": "#/components/schemas/APIStatus"
+                },
+                "gpx": {
+                  "$ref": "#/components/schemas/APIStatus"
+                }
+              }
+            }
+          }
+        },
+        "APICapabilitiesResponse": {
+          "type": "object",
+          "allOf": [
+            {
+              "$ref": "#/components/schemas/OSMResponse"
+            }
+          ],
+          "properties": {
+            "api": {
+              "$ref": "#/components/schemas/APICapabilities"
+            }
+          },
+          "xml": {
+            "wrapped": true,
+            "name": "osm"
+          }
+        },
+        "APIPermission": {
+          "type": "string",
+          "example": "allow_read_prefs",
+          "description": "Permissions:\n  * `allow_read_prefs` - Read user preferences\n  * `allow_write_prefs` - Modify user preferences\n  * `allow_write_diary` - Create diary entries, comments and make friends\n  * `allow_write_api` - Modify the map\n  * `allow_read_gpx` - Read private GPS traces\n  * `allow_write_gpx` - Upload GPS traces\n  * `allow_write_notes` - Modify notes\n",
+          "xml": {
+            "attribute": true
+          }
+        },
+        "APIStatus": {
+          "type": "string",
+          "enum": [
+            "online",
+            "readonly",
+            "offline"
+          ],
+          "example": "online",
+          "description": "Status:\n  * `online` - Both read and write requests work\n  * `readonly` - Only read requests work\n  * `offline` - No requests work\n",
+          "xml": {
+            "attribute": true
+          }
+        },
+        "APIVersionsResponse": {
+          "type": "object",
+          "allOf": [
+            {
+              "$ref": "#/components/schemas/OSMResponse"
+            }
+          ],
+          "properties": {
+            "api": {
+              "type": "object",
+              "properties": {
+                "version": {
+                  "type": "number",
+                  "format": "double",
+                  "example": 0.6
+                }
+              }
+            }
+          },
+          "xml": {
+            "wrapped": true,
+            "name": "osm"
+          }
+        },
+        "Node": {
+          "allOf": [
+            {
+              "$ref": "#/components/schemas/Way/allOf/0"
+            },
+            {
+              "type": "object",
+              "properties": {
+                "lat": {
+                  "type": "number",
+                  "format": "double",
+                  "minimum": -90,
+                  "maximum": 90,
+                  "example": 52.515785
+                },
+                "lon": {
+                  "type": "number",
+                  "format": "double",
+                  "minimum": -180,
+                  "maximum": 180,
+                  "example": 13.3776392
+                }
+              },
+              "required": [
+                "lat",
+                "lon"
+              ]
+            }
+          ]
+        },
+        "OSMResponse": {
+          "type": "object",
+          "properties": {
+            "version": {
+              "type": "string",
+              "example": "0.6",
+              "xml": {
+                "attribute": true
+              }
+            },
+            "copyright": {
+              "type": "string",
+              "example": "OpenStreetMap and contributors",
+              "xml": {
+                "attribute": true
+              }
+            },
+            "attribution": {
+              "type": "string",
+              "example": "http://www.openstreetmap.org/copyright",
+              "xml": {
+                "attribute": true
+              }
+            }
+          },
+          "required": [
+            "version",
+            "copyright",
+            "attribution"
+          ]
+        },
+        "Relation": {
+          "allOf": [
+            {
+              "$ref": "#/components/schemas/Way/allOf/0"
+            },
+            {
+              "type": "object",
+              "properties": {
+                "members": {
+                  "type": "array",
+                  "items": {
+                    "type": "object",
+                    "properties": {
+                      "type": {
+                        "$ref": "#/paths/~1api~10.6~1map/get/responses/200/content/application~1json/schema/properties/elements/items/oneOf/0/allOf/0/properties/type"
+                      },
+                      "ref": {
+                        "$ref": "#/components/schemas/Way/allOf/1/properties/nodes/items"
+                      },
+                      "role": {
+                        "type": "string",
+                        "example": "part"
+                      }
+                    }
+                  }
+                }
+              },
+              "required": [
+                "members"
+              ]
+            }
+          ]
+        },
+        "User": {
+          "type": "object",
+          "properties": {
+            "id": {
+              "type": "integer",
+              "example": 1234
+            },
+            "display_name": {
+              "type": "string",
+              "example": "JaneDoe"
+            },
+            "account_created": {
+              "type": "string",
+              "example": "2010-01-19T23:01:0Z",
+              "format": "date-time"
+            },
+            "description": {
+              "type": "string",
+              "example": "# Heading\n\nThe description can contain markdown"
+            },
+            "changesets": {
+              "type": "object",
+              "properties": {
+                "count": {
+                  "type": "integer",
+                  "example": 34
+                }
+              }
+            },
+            "traces": {
+              "type": "object",
+              "properties": {
+                "count": {
+                  "type": "integer",
+                  "example": 2
+                }
+              }
+            },
+            "messages": {
+              "type": "object",
+              "properties": {
+                "received": {
+                  "type": "object",
+                  "properties": {
+                    "count": {
+                      "type": "integer",
+                      "example": 2
+                    },
+                    "unread": {
+                      "type": "integer",
+                      "example": 1
+                    }
+                  }
+                },
+                "sent": {
+                  "type": "object",
+                  "properties": {
+                    "count": {
+                      "type": "integer",
+                      "example": 1
+                    }
+                  }
+                }
+              }
+            }
+          }
+        },
+        "UserDetailsResponse": {
+          "type": "object",
+          "allOf": [
+            {
+              "$ref": "#/components/schemas/OSMResponse"
+            }
+          ],
+          "properties": {
+            "user": {
+              "$ref": "#/components/schemas/User"
+            }
+          }
+        },
+        "Way": {
+          "allOf": [
+            {
+              "type": "object",
+              "properties": {
+                "type": {
+                  "type": "string",
+                  "enum": [
+                    "node",
+                    "way",
+                    "relation"
+                  ],
+                  "example": "node"
+                },
+                "id": {
+                  "$ref": "#/components/schemas/Way/allOf/1/properties/nodes/items"
+                },
+                "timestamp": {
+                  "type": "string",
+                  "format": "date-time",
+                  "example": "2015-10-01T19:36:42Z"
+                },
+                "version": {
+                  "type": "number",
+                  "example": 14
+                },
+                "changeset": {
+                  "type": "number",
+                  "format": "int64",
+                  "example": 34376083
+                },
+                "user": {
+                  "type": "string",
+                  "example": "anbr"
+                },
+                "uid": {
+                  "type": "number",
+                  "example": 43566
+                },
+                "tags": {
+                  "type": "object",
+                  "additionalProperties": {
+                    "type": "string"
+                  },
+                  "example": {
+                    "building:architecture": "neoclassicism",
+                    "building:material": "sandstone"
+                  }
+                }
+              },
+              "required": [
+                "type",
+                "id",
+                "timestamp",
+                "version",
+                "changeset",
+                "user",
+                "uid"
+              ]
+            },
+            {
+              "type": "object",
+              "properties": {
+                "nodes": {
+                  "type": "array",
+                  "items": {
+                    "type": "number",
+                    "format": "int64",
+                    "example": 85883220
+                  }
+                }
+              },
+              "required": [
+                "nodes"
+              ]
+            }
+          ]
+        }
+      }
+    },
+    "externalDocs": {
+      "description": "Find more information  n the OSM wiki",
+      "url": "https://wiki.openstreetmap.org/"
+    }
+  }
\ No newline at end of file
