WebAPI missing files pt2

diff --git a/Engine/Plugins/Web/WebAPI/Source/WebAPIOpenAPI/Private/Tests/WebAPISwagger.spec.cpp b/Engine/Plugins/Web/WebAPI/Source/WebAPIOpenAPI/Private/Tests/WebAPISwagger.spec.cpp
new file mode 100644
--- /dev/null
+++ b/Engine/Plugins/Web/WebAPI/Source/WebAPIOpenAPI/Private/Tests/WebAPISwagger.spec.cpp
@@ -0,0 +1,370 @@
+﻿// Copyright Epic Games, Inc. All Rights Reserved.
+
+#if WITH_EDITOR
+#if WITH_DEV_AUTOMATION_TESTS
+
+#include "WebAPIEditorSettings.h"
+#include "Dom/WebAPIParameter.h"
+#include "V2/WebAPISwaggerConverter.h"
+#include "V2/WebAPISwaggerFactory.h"
+#include "V2/WebAPISwaggerSchema.h"
+
+using namespace UE::WebAPI::OpenAPI;
+
+BEGIN_DEFINE_SPEC(FWebAPISwaggerSpec,
+				TEXT("Plugin.WebAPI.Swagger"),
+				EAutomationTestFlags::EditorContext | EAutomationTestFlags::EngineFilter | EAutomationTestFlags::ApplicationContextMask)
+
+	TSharedPtr<UE::WebAPI::OpenAPI::V2::FSwagger> InputDefinition;
+	TStrongObjectPtr<UWebAPIDefinition> OutputDefinition;
+	TSharedPtr<UE::WebAPI::Swagger::FWebAPISwaggerSchemaConverter> Converter;
+	TStrongObjectPtr<UWebAPIDefinitionFactory> Factory;
+
+	FString GetSampleFile(const FString& InName) const
+	{
+		FString FilePath = FPaths::Combine(FPaths::ProjectPluginsDir(),
+			TEXT("WebAPI"), TEXT("Source"), TEXT("WebAPIOpenAPI"),
+			TEXT("Private"), TEXT("Tests"), TEXT("Samples"), TEXT("V2"), InName + TEXT(".json"));
+		ensure(FPaths::FileExists(FilePath));
+		return FilePath;
+	}
+
+	FString GetAPISample() const
+	{
+		return GetSampleFile(TEXT("petstore_V2.json"));
+	}
+
+	FString GetDigitalTwinsManagementAPISample() const
+	{
+		return GetSampleFile(TEXT("digitaltwins"));
+	}
+
+	FString GetDigitalTwinsAPISample() const
+	{
+		return GetSampleFile(TEXT("digitaltwins_data"));
+	}
+
+	TSharedPtr<FJsonObject> LoadJson(const FString& InFile) const
+	{
+		FString FileContents;
+		FFileHelper::LoadFileToString(FileContents, *InFile);
+	
+		TSharedPtr<FJsonObject> JsonObject;
+		FJsonSerializer::Deserialize(TJsonReaderFactory<TCHAR>::Create(FileContents), JsonObject);
+		return JsonObject;
+	}
+
+END_DEFINE_SPEC(FWebAPISwaggerSpec)
+
+// @todo: remove when done: -ExecCmds="Automation RunTests Plugin.WebAPI.Parser.Swagger2" -testexit="Automation Test Queue Empty" -unattended  -nopause
+void FWebAPISwaggerSpec::Define()
+{
+	BeforeEach([this]
+	{
+		const FString FilePath = GetDigitalTwinsManagementAPISample();
+		const TSharedPtr<FJsonObject> JsonObject = LoadJson(FilePath);
+
+		InputDefinition = MakeShared<UE::WebAPI::OpenAPI::V2::FSwagger>();
+		InputDefinition->FromJson(JsonObject.ToSharedRef());
+					
+		OutputDefinition = TStrongObjectPtr(NewObject<UWebAPIDefinition>());
+
+		Converter = MakeShared<UE::WebAPI::Swagger::FWebAPISwaggerSchemaConverter>(
+			InputDefinition,
+			OutputDefinition->GetWebAPISchema(),
+			OutputDefinition->GetMessageLog().ToSharedRef(),
+			OutputDefinition->GetProviderSettings());
+
+		Factory = TStrongObjectPtr(NewObject<UWebAPISwaggerFactory>());
+	});
+	
+	Describe("PetStore", [this]
+	{
+		It("Parses", [this]
+		{
+			const FString FilePath = GetAPISample();
+			const TSharedPtr<FJsonObject> JsonObject = LoadJson(FilePath);
+
+			UE::WebAPI::OpenAPI::V2::FSwagger Object;
+			Object.FromJson(JsonObject.ToSharedRef());
+
+			TestEqual("Swagger Version", Object.Swagger, "2.0");
+			TestEqual("Info Version", Object.Info->Version, "1.0.5");
+			if(TestTrue("Host Set", Object.Host.IsSet()))
+			{
+				TestEqual("Host", Object.Host.GetValue(), "petstore.swagger.io");
+			}
+
+			if(TestTrue("BasePath Set", Object.BasePath.IsSet()))
+			{
+				TestEqual("BasePath", Object.BasePath.GetValue(), "/v2");	
+			}
+
+			if(TestTrue("Tags Set", Object.Tags.IsSet()))
+			{
+				TestEqual("Tag Num", Object.Tags->Num(), 3);	
+			}
+
+			if(TestTrue("Schemes Set", Object.Schemes.IsSet()))
+			{
+				TestEqual("Scheme Num", Object.Schemes->Num(), 2);
+			}
+
+			TestEqual("Paths Num", Object.Paths.Num(), 14);
+			
+			if(TestTrue("Security Definitions Set", Object.SecurityDefinitions.IsSet()))
+			{
+				TestEqual("Security Definitions Num", Object.SecurityDefinitions->Num(), 2);
+			}
+			
+			if(TestTrue("Definitions Set", Object.Definitions.IsSet()))
+			{
+				TestEqual("Definitions Num", Object.Definitions->Num(), 6);
+			}
+
+			if(TestTrue("External Docs Set", Object.ExternalDocs.IsSet()))
+			{
+				TestEqual("External Docs URL", Object.ExternalDocs.GetValue()->Url, TEXT("http://swagger.io"));
+			}			
+		});
+	});
+
+	Describe("DigitalTwinsManagement", [this]
+	{
+		It("Parses", [this]
+		{
+			const FString FilePath = GetDigitalTwinsManagementAPISample();
+			const TSharedPtr<FJsonObject> JsonObject = LoadJson(FilePath);
+
+			UE::WebAPI::OpenAPI::V2::FSwagger Object;
+			Object.FromJson(JsonObject.ToSharedRef());
+
+			TestEqual("Swagger Version", Object.Swagger, "2.0");
+			TestEqual("Info Version", Object.Info->Version, "2020-12-01");
+			if(TestTrue("Host Set", Object.Host.IsSet()))
+			{
+				TestEqual("Host", Object.Host.GetValue(), "management.azure.com");
+			}
+
+			if(TestTrue("Schemes Set", Object.Schemes.IsSet()))
+			{
+				TestEqual("Scheme Num", Object.Schemes->Num(), 1);
+			}
+
+			TestEqual("Paths Num", Object.Paths.Num(), 11);
+				
+			if(TestTrue("Security Definitions Set", Object.SecurityDefinitions.IsSet()))
+			{
+				TestEqual("Security Definitions Num", Object.SecurityDefinitions->Num(), 1);
+			}
+				
+			if(TestTrue("Definitions Set", Object.Definitions.IsSet()))
+			{
+				TestEqual("Definitions Num", Object.Definitions->Num(), 30);
+			}
+		});
+	});
+
+	Describe("Model Composition", [this]
+	{
+		// @todo: meaningful name
+		It("All-Of", [this]
+		{
+			const TSharedPtr<UE::WebAPI::OpenAPI::V2::FSchema>* InputModel = InputDefinition->Definitions->Find(TEXT("ServiceBus"));
+
+			const TObjectPtr<UWebAPIModel> OutputModel = Converter->ConvertModel(*InputModel);
+			if(TestFalse("Converted Model", OutputModel.IsNull()))
+			{
+				if(TestTrue("All-Of is set", (*InputModel)->AllOf.IsSet()))
+				{
+					for(UE::Json::TJsonReference<UE::WebAPI::OpenAPI::V2::FSchema>& SubType : (*InputModel)->AllOf.GetValue())
+					{
+					
+					}
+				}
+			}
+		});
+	});
+
+	Describe("Type Conversion", [this]
+	{
+		BeforeEach([this]
+		{
+			const FString FilePath = GetDigitalTwinsAPISample();
+			const TSharedPtr<FJsonObject> JsonObject = LoadJson(FilePath);
+
+			InputDefinition = MakeShared<UE::WebAPI::OpenAPI::V2::FSwagger>();
+			InputDefinition->FromJson(JsonObject.ToSharedRef());
+					
+			OutputDefinition = TStrongObjectPtr(NewObject<UWebAPIDefinition>());
+
+			Converter = MakeShared<UE::WebAPI::Swagger::FWebAPISwaggerSchemaConverter>(
+				InputDefinition,
+				OutputDefinition->GetWebAPISchema(),
+				OutputDefinition->GetMessageLog().ToSharedRef(),
+				OutputDefinition->GetProviderSettings());
+
+			Factory = TStrongObjectPtr(NewObject<UWebAPISwaggerFactory>());
+		});
+		// // @todo: meaningful name
+		// It("Test", [this]()
+		// {
+		// 	const TSharedPtr<UE::WebAPI::OpenAPI::V2::FSchema>* DigitalTwinsProperties = InputDefinition->Definitions->Find(TEXT("DigitalTwinsProperties"));
+		//
+		// 	const TObjectPtr<UWebAPITypeInfo> Type = Converter->ResolveType(*DigitalTwinsProperties, TEXT("DigitalTwinsProperties"));
+		//
+		// 	TestTrue("TypeInfo is not null", Type != nullptr);
+		// });
+
+		It("Enum that's non-optional and has a single value should default to that value", [this]
+		{
+			check(Converter->Convert());
+			Factory->PostImportWebAPI(OutputDefinition.Get()).Wait();
+			
+			const FString EnumNameToTest = TEXT("ApiVersion");
+
+			TObjectPtr<UWebAPIModelBase>* FoundModel = OutputDefinition->GetWebAPISchema()->Models.FindByPredicate([&](const TObjectPtr<UWebAPIModelBase>& InModel)
+			{
+				return InModel->IsA(UWebAPIEnum::StaticClass()) && InModel->GetName().Contains(EnumNameToTest);
+			});
+
+			if(!TestNotNull("Found Enum", FoundModel)
+				|| !TestTrue("Found Enum not nullptr", !FoundModel->IsNull()))
+			{
+				return false;				
+			}
+
+			const TObjectPtr<UWebAPIEnum> FoundEnum = Cast<UWebAPIEnum>(*FoundModel);
+			TestEqual("Single value", FoundEnum->Values.Num(), 1);
+			TestTrue("Default value is the first value", FoundEnum->GetDefaultValue().Contains(FoundEnum->Values[0]->Name.ToMemberName(TEXT("EV"))));
+
+ 			return true;
+		});
+
+		It("String Array that's optional should have an empty array as a default value", [this]
+		{
+			check(Converter->Convert());
+			Factory->PostImportWebAPI(OutputDefinition.Get()).Wait();
+			
+			const FString ObjectNameToTest = TEXT("IncomingRelationshipCollection");
+			const FString PropertyNameToTest = TEXT("Value");
+
+			TObjectPtr<UWebAPIModelBase>* FoundModelBase = OutputDefinition->GetWebAPISchema()->Models.FindByPredicate([&](const TObjectPtr<UWebAPIModelBase>& InModel)
+			{
+				return InModel->IsA(UWebAPIModel::StaticClass()) && InModel->GetName().Contains(ObjectNameToTest);
+			});
+
+			if(!TestNotNull("Found Model", FoundModelBase)
+				|| !TestTrue("Found Model not nullptr", !FoundModelBase->IsNull()))
+			{
+				return false;				
+			}
+
+			const TObjectPtr<UWebAPIModel> FoundModel = Cast<UWebAPIModel>(*FoundModelBase);
+
+			const TObjectPtr<UWebAPIProperty>* StringArrayProperty = FoundModel->Properties.FindByPredicate([&](const TObjectPtr<UWebAPIProperty>& InProperty)
+			{
+				return InProperty->Name.ToString(true) == PropertyNameToTest;				
+			});
+
+			const FString DefaultValue = (*StringArrayProperty)->GetDefaultValue();
+			TestTrue("Empty Array default value", DefaultValue.Len() <= 3);
+
+			return true;
+		});
+		
+		It("Default value should be set in metadata", [this]
+		{
+
+		});
+
+		It("Properties marked as required in the root of a definition are required in the output object(s)", [this]
+		{
+			check(Converter->Convert());
+			Factory->PostImportWebAPI(OutputDefinition.Get()).Wait();
+					
+			const FString EnumNameToTest = TEXT("ApiVersion");
+
+			TObjectPtr<UWebAPIModelBase>* FoundModel = OutputDefinition->GetWebAPISchema()->Models.FindByPredicate([&](const TObjectPtr<UWebAPIModelBase>& InModel)
+			{
+				return InModel->GetName().Contains(EnumNameToTest);
+			});
+
+			if(!TestNotNull("Found Enum", FoundModel)
+				|| TestTrue("Found Enum not nullptr", !FoundModel->IsNull()))
+			{
+				return false;				
+			}
+
+			const TObjectPtr<UWebAPIEnum> FoundEnum = Cast<UWebAPIEnum>(*FoundModel);
+			const TObjectPtr<UWebAPIEnumValue>* FoundEnumValue = FoundEnum->Values.FindByPredicate([&](const TObjectPtr<UWebAPIEnumValue>& InEnumValue)
+			{
+				return InEnumValue->Name == OutputDefinition->GetProviderSettings().GetUnsetEnumValueName();
+			});
+
+			TestNull("Unset value not found", FoundEnumValue);
+
+			return true;
+		});
+	});
+
+	Describe("Parameter", [this]
+	{
+		It("Resolves schema-less, built-in type", [this]
+		{
+			const TSharedPtr<UE::WebAPI::OpenAPI::V2::FParameter>* InputParameter = InputDefinition->Parameters->Find(TEXT("resourceName"));
+
+			const TObjectPtr<UWebAPIParameter> OutputParameter = Converter->ConvertParameter(*InputParameter);
+			if(TestFalse("Converted Parameter", OutputParameter.IsNull()))
+			{
+				TestEqual("Type is String", OutputParameter->Type.ToString(true), TEXT("String"));
+				TestTrue("Object is named correctly", OutputParameter->GetFName().ToString().StartsWith(TEXT("WebAPIParameter_ResourceName")));
+			}
+		});
+
+		It("Resolves schema-less, generated enum type", [this]
+		{
+			const TSharedPtr<UE::WebAPI::OpenAPI::V2::FParameter>* InputParameter = InputDefinition->Parameters->Find(TEXT("api-version"));
+
+			const TObjectPtr<UWebAPIParameter> OutputParameter = Converter->ConvertParameter(*InputParameter);
+			if(TestFalse("Converted Parameter", OutputParameter.IsNull()))
+			{
+				TestTrue("Type is Enum", OutputParameter->Type.TypeInfo->IsEnum());
+
+				UWebAPIEnum* Model = Cast<UWebAPIEnum>(OutputParameter->Type.TypeInfo->Model.Get());
+				if(TestTrue("Underlying Enum", Model != nullptr))
+				{
+					// TestTrue("Type is Enum", OutputParameter->Type.TypeInfo->IsEnum());
+				}
+			}
+		});
+	});
+		
+	Describe("Operation", [this]
+	{
+		It("Resolves referenced parameters", [this]
+		{
+			const FString Url = TEXT("/subscriptions/{subscriptionId}/resourceGroups/{resourceGroupName}/providers/Microsoft.DigitalTwins/digitalTwinsInstances/{resourceName}/privateEndpointConnections/{privateEndpointConnectionName}");
+
+			const TSharedPtr<V2::FPath, ESPMode::ThreadSafe>* InputService = InputDefinition->Paths.Find(Url);
+			const FString ServiceTag = TEXT("PrivateEndpoints");
+			const TObjectPtr<UWebAPIService> OutputService = Converter->ConvertService(ServiceTag);
+			
+			const TSharedPtr<UE::WebAPI::OpenAPI::V2::FOperation>& InputOperation = InputDefinition->Paths
+				.Find(Url)
+				->Get()->Get;
+
+			const TObjectPtr<UWebAPIOperation> OutputOperation = Converter->ConvertOperation(Url, TEXT("Get"), InputOperation);
+			const FString OperationName = OutputOperation->Name.ToString(true);
+			
+			if(TestFalse("Converted Operation", OutputOperation.IsNull()))
+			{
+				TestTrue("Object is named correctly", OutputOperation->GetFName().ToString().StartsWith(TEXT("WebAPIOperation_PrivateEndpoints_") + OperationName));
+			}
+		});
+
+	});
+}
+
+#endif
+#endif
diff --git a/Engine/Plugins/Web/WebAPI/Source/WebAPIOpenAPI/Private/V2/WebApiSwaggerFactory.cpp b/Engine/Plugins/Web/WebAPI/Source/WebAPIOpenAPI/Private/V2/WebApiSwaggerFactory.cpp
new file mode 100644
--- /dev/null
+++ b/Engine/Plugins/Web/WebAPI/Source/WebAPIOpenAPI/Private/V2/WebApiSwaggerFactory.cpp
@@ -0,0 +1,61 @@
+﻿// Copyright Epic Games, Inc. All Rights Reserved.
+
+#include "WebAPISwaggerFactory.h"
+
+#include "WebAPIOpenAPILog.h"
+#include "WebAPIDefinition.h"
+#include "Dom/JsonObject.h"
+#include "Dom/WebAPIOperation.h"
+#include "Serialization/JsonSerializer.h"
+
+#define LOCTEXT_NAMESPACE "WebAPISwaggerFactory"
+
+struct FWebAPIOperationResponseKeyFuncs : public TDefaultMapKeyFuncs<uint32, TObjectPtr<UWebAPIOperationResponse>, false>
+{
+	static FORCEINLINE uint32 GetSetKey(TObjectPtr<UWebAPIOperationResponse> const& Element) { return Element->Code; }
+	static FORCEINLINE uint32 GetKeyHash(uint32 const& Key) { return Key; }
+	static FORCEINLINE bool Matches(uint32 const& A, uint32 const& B) { return (A == B); }
+};
+
+UWebAPISwaggerFactory::UWebAPISwaggerFactory()
+	: Provider(MakeShared<FWebAPISwaggerProvider>())
+{
+}
+
+bool UWebAPISwaggerFactory::CanImportWebAPI(const FString& InFileName, const FString& InFileContents)
+{
+	TSharedPtr<FJsonObject> JsonObject;
+	if(!FJsonSerializer::Deserialize(TJsonReaderFactory<TCHAR>::Create(InFileContents), JsonObject))
+	{
+		return false;
+	}
+
+	return JsonObject->HasField(TEXT("swagger"));
+}
+
+TFuture<bool> UWebAPISwaggerFactory::ImportWebAPI(UWebAPIDefinition* InDefinition, const FString& InFileName, const FString& InFileContents)
+{
+	TSharedPtr<FJsonObject> JsonObject;
+	if(!FJsonSerializer::Deserialize(TJsonReaderFactory<TCHAR>::Create(InFileContents), JsonObject))
+	{
+		UE_LOG(LogWebAPIOpenAPI, Error, TEXT("Couldn't deserialize file contents as Json"));
+		return MakeFulfilledPromise<bool>(false).GetFuture();
+	}
+
+	// Write raw spec contents
+	UWebAPISwaggerAssetData* ImportData = InDefinition->AddOrGetImportedDataCache<UWebAPISwaggerAssetData>("Swagger");
+	ImportData->FileContents = InFileContents;
+
+	// Parse spec
+	// @todo: move call to ConvertToWebAPISchema outside of Factory implementation (to provide Pre and Post extensibility)
+	return Provider->ConvertToWebAPISchema(InDefinition)
+	.Next([](EWebAPIConversionResult bInConversionResult)
+	{
+		// @todo: UI notification
+		//InDefinition->SetWebAPISchema(InWebAPISchema);
+
+		return bInConversionResult == EWebAPIConversionResult::Succeeded;
+	});
+}
+
+#undef LOCTEXT_NAMESPACE
diff --git a/Engine/Plugins/Web/WebAPI/Source/WebAPIOpenAPI/Private/V2/WebApiSwaggerSchema.cpp b/Engine/Plugins/Web/WebAPI/Source/WebAPIOpenAPI/Private/V2/WebApiSwaggerSchema.cpp
new file mode 100644
--- /dev/null
+++ b/Engine/Plugins/Web/WebAPI/Source/WebAPIOpenAPI/Private/V2/WebApiSwaggerSchema.cpp
@@ -0,0 +1,322 @@
+﻿// Copyright Epic Games, Inc. All Rights Reserved.
+
+#include "V2/WebAPISwaggerSchema.h"
+
+#include "WebAPIJsonUtilities.h"
+#include "Algo/ForEach.h"
+#include "Dom/JsonObject.h"
+
+using namespace UE;
+
+namespace UE::WebAPI::OpenAPI::V2
+{
+	static TMap<FString, EJson> EJson_NameToValue = {
+		{TEXT("None"), EJson::None},
+		{TEXT("Null"), EJson::Null},
+		{TEXT("String"), EJson::String},
+		{TEXT("Boolean"), EJson::Boolean},
+		{TEXT("Number"), EJson::Number},
+		{TEXT("Integer"), EJson::Number},
+		{TEXT("Array"), EJson::Array},
+		{TEXT("Object"), EJson::Object},
+	};
+
+	bool FItems::FromJson(const TSharedRef<FJsonObject>& InJson)
+	{
+		bool bSetAField = false;
+		bSetAField &= Json::TryGetField(InJson, TEXT("type"), Type);
+		bSetAField &= Json::TryGetField(InJson, TEXT("format"), Format);
+		//bSetAField &= Json::TryGetField(InJson, TEXT("items"), Items); // @todo: fix recursion
+		bSetAField &= Json::TryGetField(InJson, TEXT("collectionFormat"), CollectionFormat);
+		bSetAField &= Json::TryGetField(InJson, TEXT("default"), Default);
+		bSetAField &= Json::TryGetField(InJson, TEXT("maximum"), Maximum);
+		bSetAField &= Json::TryGetField(InJson, TEXT("exclusiveMaximum"), ExclusiveMaximum);
+		bSetAField &= Json::TryGetField(InJson, TEXT("minimum"), Minimum);
+		bSetAField &= Json::TryGetField(InJson, TEXT("exclusiveMinimum"), ExclusiveMinimum);
+		bSetAField &= Json::TryGetField(InJson, TEXT("maxLength"), MaxLength);
+		bSetAField &= Json::TryGetField(InJson, TEXT("minLength"), MinLength);
+		bSetAField &= Json::TryGetField(InJson, TEXT("pattern"), Pattern);
+		bSetAField &= Json::TryGetField(InJson, TEXT("minItems"), MinItems);
+		bSetAField &= Json::TryGetField(InJson, TEXT("maxItems"), MaxItems);
+		bSetAField &= Json::TryGetField(InJson, TEXT("uniqueItems"), bUniqueItems);
+		bSetAField &= Json::TryGetField(InJson, TEXT("enum"), Enum);
+		bSetAField &= Json::TryGetField(InJson, TEXT("multipleOf"), MultipleOf);
+		return bSetAField;
+	}
+
+	bool FExternalDocumentation::FromJson(const TSharedRef<FJsonObject>& InJson)
+	{
+		bool bSetAField = false;
+		bSetAField &= Json::TryGetField(InJson, TEXT("description"), Description);
+		bSetAField &= Json::TryGetField(InJson, TEXT("url"), Url);
+		return bSetAField;
+	}
+
+	bool FXML::FromJson(const TSharedRef<FJsonObject>& InJson)
+	{
+		bool bSetAField = false;
+		bSetAField &= Json::TryGetField(InJson, TEXT("name"), Name);
+		bSetAField &= Json::TryGetField(InJson, TEXT("namespace"), Namespace);
+		bSetAField &= Json::TryGetField(InJson, TEXT("prefix"), Prefix);
+		bSetAField &= Json::TryGetField(InJson, TEXT("attribute"), Attribute);
+		bSetAField &= Json::TryGetField(InJson, TEXT("wrapped"), Wrapped);
+		return bSetAField;
+	}
+
+	bool FSchemaBase::FromJson(const TSharedRef<FJsonObject>& InJson)
+	{
+		bool bSetAField = false;
+		bSetAField &= Json::TryGetField(InJson, TEXT("name"), Name);
+		bSetAField &= Json::TryGetField(InJson, TEXT("description"), Description);
+		bSetAField &= Json::TryGetField(InJson, TEXT("default"), Default);
+		bSetAField &= Json::TryGetField(InJson, TEXT("maximum"), Maximum);
+		bSetAField &= Json::TryGetField(InJson, TEXT("exclusiveMaximum"), bExclusiveMaximum);
+		bSetAField &= Json::TryGetField(InJson, TEXT("minimum"), Minimum);
+		bSetAField &= Json::TryGetField(InJson, TEXT("exclusiveMinimum"), bExclusiveMinimum);
+		bSetAField &= Json::TryGetField(InJson, TEXT("maxLength"), MaxLength);
+		bSetAField &= Json::TryGetField(InJson, TEXT("minLength"), MinLength);
+		bSetAField &= Json::TryGetField(InJson, TEXT("pattern"), Pattern);
+		bSetAField &= Json::TryGetField(InJson, TEXT("maxItems"), MaxItems);
+		bSetAField &= Json::TryGetField(InJson, TEXT("minItems"), MinItems);
+		bSetAField &= Json::TryGetField(InJson, TEXT("uniqueItems"), bUniqueItems);
+		bSetAField &= Json::TryGetField(InJson, TEXT("required"), bRequired);
+		bSetAField &= Json::TryGetField(InJson, TEXT("enum"), Enum);
+		bSetAField &= Json::TryGetField(InJson, TEXT("type"), Type);
+		bSetAField &= Json::TryGetField(InJson, TEXT("items"), Items);
+		return bSetAField;
+	}
+
+	bool FSchema::FromJson(const TSharedRef<FJsonObject>& InJson)
+	{
+		bool bSetAField = Super::FromJson(InJson);
+		bSetAField &= Json::TryGetField(InJson, TEXT("format"), Format);
+		bSetAField &= Json::TryGetField(InJson, TEXT("title"), Title);
+		bSetAField &= Json::TryGetField(InJson, TEXT("allOf"), AllOf);
+		bSetAField &= Json::TryGetField(InJson, TEXT("properties"), Properties);
+		bSetAField &= Json::TryGetField(InJson, TEXT("additionalProperties"), AdditionalProperties);
+		bSetAField &= Json::TryGetField(InJson, TEXT("discriminator"), Discriminator);
+		bSetAField &= Json::TryGetField(InJson, TEXT("readOnly"), bReadOnly);
+		bSetAField &= Json::TryGetField(InJson, TEXT("xml"), XML);
+		bSetAField &= Json::TryGetField(InJson, TEXT("externalDocs"), ExternalDocs);
+		bSetAField &= Json::TryGetField(InJson, TEXT("example"), Example);
+		return bSetAField;
+	}
+
+	bool FParameter::FromJson(const TSharedRef<FJsonObject>& InJson)
+	{
+		bool bSetAField = Super::FromJson(InJson);
+		bSetAField &= Json::TryGetField(InJson, TEXT("in"), In);
+		bSetAField &= Json::TryGetField(InJson, TEXT("description"), Description);
+		bSetAField &= Json::TryGetField(InJson, TEXT("format"), Format);
+		bSetAField &= Json::TryGetField(InJson, TEXT("schema"), Schema);
+		bSetAField &= Json::TryGetField(InJson, TEXT("multipleOf"), MultipleOf);
+		return bSetAField;
+	}
+
+	bool FHeader::FromJson(const TSharedRef<FJsonObject>& InJson)
+	{
+		bool bSetAField = false;
+		bSetAField &= Json::TryGetField(InJson, TEXT("description"), Description);
+		bSetAField &= Json::TryGetField(InJson, TEXT("type"), Type);
+		bSetAField &= Json::TryGetField(InJson, TEXT("format"), Format);
+		bSetAField &= Json::TryGetField(InJson, TEXT("items"), Items);
+		bSetAField &= Json::TryGetField(InJson, TEXT("collectionFormat"), CollectionFormat);
+		bSetAField &= Json::TryGetField(InJson, TEXT("default"), Default);
+		bSetAField &= Json::TryGetField(InJson, TEXT("maximum"), Maximum);
+		bSetAField &= Json::TryGetField(InJson, TEXT("exclusiveMaximum"), bExclusiveMaximum);
+		bSetAField &= Json::TryGetField(InJson, TEXT("minimum"), Minimum);
+		bSetAField &= Json::TryGetField(InJson, TEXT("exclusiveMinimum"), bExclusiveMinimum);
+		bSetAField &= Json::TryGetField(InJson, TEXT("maxLength"), MaxLength);
+		bSetAField &= Json::TryGetField(InJson, TEXT("minLength"), MinLength);
+		bSetAField &= Json::TryGetField(InJson, TEXT("pattern"), Pattern);
+		bSetAField &= Json::TryGetField(InJson, TEXT("maxItems"), MaxItems);
+		bSetAField &= Json::TryGetField(InJson, TEXT("minItems"), MinItems);
+		bSetAField &= Json::TryGetField(InJson, TEXT("uniqueItems"), bUniqueItems);
+		bSetAField &= Json::TryGetField(InJson, TEXT("enum"), Enum);
+		bSetAField &= Json::TryGetField(InJson, TEXT("multipleOf"), MultipleOf);
+		return bSetAField;
+	}
+
+	bool FExample::FromJson(const TSharedRef<FJsonObject>& InJson)
+	{
+		// @todo: implement
+		const bool bSetAField = false;
+		return bSetAField;
+	}
+
+	bool FResponse::FromJson(const TSharedRef<FJsonObject>& InJson)
+	{
+		bool bSetAField = false;
+		bSetAField &= Json::TryGetField(InJson, TEXT("description"), Description);
+		bSetAField &= Json::TryGetField(InJson, TEXT("schema"), Schema);
+		bSetAField &= Json::TryGetField(InJson, TEXT("headers"), Headers);
+		bSetAField &= Json::TryGetField(InJson, TEXT("examples"), Examples);
+		return bSetAField;
+	}
+
+	bool FScopes::FromJson(const TSharedRef<FJsonObject>& InJson)
+	{
+		// @todo: implement
+		const bool bSetAField = false;
+		return bSetAField;
+	}
+
+	bool FSecurityScheme::FromJson(const TSharedRef<FJsonObject>& InJson)
+	{
+		bool bSetAField = false;
+		bSetAField &= Json::TryGetField(InJson, TEXT("type"), Type);
+		bSetAField &= Json::TryGetField(InJson, TEXT("description"), Description);
+		bSetAField &= Json::TryGetField(InJson, TEXT("name"), Name);
+		bSetAField &= Json::TryGetField(InJson, TEXT("in"), In);
+		bSetAField &= Json::TryGetField(InJson, TEXT("flow"), Flow);
+		bSetAField &= Json::TryGetField(InJson, TEXT("authorizationUrl"), AuthorizationUrl);
+		bSetAField &= Json::TryGetField(InJson, TEXT("tokenUrl"), TokenUrl);
+		bSetAField &= Json::TryGetField(InJson, TEXT("scopes"), Scopes);
+		return bSetAField;
+	}
+
+	bool FSecurityRequirement::FromJson(const TSharedRef<FJsonObject>& InJson)
+	{
+		const bool bSetAField = false;
+		return bSetAField;
+	}
+
+	bool FTag::FromJson(const TSharedRef<FJsonObject>& InJson)
+	{
+		bool bSetAField = false;
+		bSetAField &= Json::TryGetField(InJson, TEXT("name"), Name);
+		bSetAField &= Json::TryGetField(InJson, TEXT("description"), Description);
+		bSetAField &= Json::TryGetField(InJson, TEXT("externalDocs"), ExternalDocs);
+		return bSetAField;
+	}
+
+	bool FPath::FromJson(const TSharedRef<FJsonObject>& InJson)
+	{
+		bool bSetAField = false;
+		bSetAField &= Json::TryGetField(InJson, TEXT("get"), Get);
+		bSetAField &= Json::TryGetField(InJson, TEXT("put"), Put);
+		bSetAField &= Json::TryGetField(InJson, TEXT("post"), Post);
+		bSetAField &= Json::TryGetField(InJson, TEXT("delete"), Delete);
+		bSetAField &= Json::TryGetField(InJson, TEXT("options"), Options);
+		bSetAField &= Json::TryGetField(InJson, TEXT("head"), Head);
+		bSetAField &= Json::TryGetField(InJson, TEXT("patch"), Patch);
+		bSetAField &= Json::TryGetField(InJson, TEXT("parameters"), Parameters);
+		return bSetAField;
+	}
+
+	bool FSwagger::FromJson(const TSharedRef<FJsonObject>& InJson)
+	{
+		bool bSetAField = false;
+
+		bSetAField &= Json::TryGetField(InJson, TEXT("swagger"), Swagger);
+		bSetAField &= Json::TryGetField(InJson, TEXT("definitions"), Definitions);
+		bSetAField &= Json::TryGetField(InJson, TEXT("info"), Info);
+		bSetAField &= Json::TryGetField(InJson, TEXT("host"), Host);
+		bSetAField &= Json::TryGetField(InJson, TEXT("basePath"), BasePath);
+		bSetAField &= Json::TryGetField(InJson, TEXT("schemes"), Schemes);
+		bSetAField &= Json::TryGetField(InJson, TEXT("consumes"), Consumes);
+		bSetAField &= Json::TryGetField(InJson, TEXT("produces"), Produces);
+		bSetAField &= Json::TryGetField(InJson, TEXT("paths"), Paths);
+		bSetAField &= Json::TryGetField(InJson, TEXT("parameters"), Parameters);
+		bSetAField &= Json::TryGetField(InJson, TEXT("responses"), Responses);
+		bSetAField &= Json::TryGetField(InJson, TEXT("securityDefinitions"), SecurityDefinitions);
+		bSetAField &= Json::TryGetField(InJson, TEXT("security"), Security);
+		bSetAField &= Json::TryGetField(InJson, TEXT("tags"), Tags);
+		bSetAField &= Json::TryGetField(InJson, TEXT("externalDocs"), ExternalDocs);
+
+		Algo::ForEachIf(Definitions.Get({}),
+		[](const TPair<FString, TSharedPtr<FSchema>>& KVP)
+		{
+			return KVP.Value->Name.IsEmpty();
+		},
+		[](const TPair<FString, TSharedPtr<FSchema>>& KVP)
+		{
+			KVP.Value->Name = KVP.Key;				
+		});
+
+		Algo::ForEachIf(Parameters.Get({}),
+		[](const TPair<FString, TSharedPtr<FParameter>>& KVP)
+		{
+			return KVP.Value->Name.IsEmpty();
+		},
+		[](const TPair<FString, TSharedPtr<FParameter>>& KVP)
+		{
+			KVP.Value->Name = KVP.Key;				
+		});
+
+		Algo::ForEachIf(Responses.Get({}),
+		[](const TPair<FString, TSharedPtr<FResponse>>& KVP)
+		{
+			return !KVP.Value->Name.IsSet();
+		},
+		[](const TPair<FString, TSharedPtr<FResponse>>& KVP)
+		{
+			KVP.Value->Name = KVP.Key;		
+		});
+
+		Algo::ForEachIf(SecurityDefinitions.Get({}),
+		[](const TPair<FString, TSharedPtr<FSecurityScheme>>& KVP)
+		{
+			return !KVP.Value->Name.IsSet();
+		},
+		[](const TPair<FString, TSharedPtr<FSecurityScheme>>& KVP)
+		{
+			KVP.Value->Name = KVP.Key;		
+		});
+
+		return bSetAField;
+	}
+
+	bool FLicense::FromJson(const TSharedRef<FJsonObject>& InJson)
+	{
+		bool bSetAField = false;
+		bSetAField &= Json::TryGetField(InJson, TEXT("name"), Name);
+		bSetAField &= Json::TryGetField(InJson, TEXT("url"), URL);
+		return bSetAField;
+	}
+
+	bool FContact::FromJson(const TSharedRef<FJsonObject>& InJson)
+	{
+		bool bSetAField = false;
+		bSetAField &= Json::TryGetField(InJson, TEXT("name"), Name);
+		bSetAField &= Json::TryGetField(InJson, TEXT("url"), URL);
+		bSetAField &= Json::TryGetField(InJson, TEXT("email"), Email);
+		return bSetAField;
+	}
+
+	bool FInfo::FromJson(const TSharedRef<FJsonObject>& InJson)
+	{
+		bool bSetAField = false;
+		bSetAField &= Json::TryGetField(InJson, TEXT("title"), Title);
+		bSetAField &= Json::TryGetField(InJson, TEXT("description"), Description);
+		bSetAField &= Json::TryGetField(InJson, TEXT("termsOfService"), TermsOfService);
+		bSetAField &= Json::TryGetField(InJson, TEXT("contact"), Contact);
+		bSetAField &= Json::TryGetField(InJson, TEXT("license"), License);
+		bSetAField &= Json::TryGetField(InJson, TEXT("version"), Version);
+		return bSetAField;
+	}
+
+	bool FOperation::FromJson(const TSharedRef<FJsonObject>& InJson)
+	{
+		bool bSetAField = false;
+		bSetAField &= Json::TryGetField(InJson, TEXT("tags"), Tags);
+		bSetAField &= Json::TryGetField(InJson, TEXT("summary"), Summary);
+		bSetAField &= Json::TryGetField(InJson, TEXT("description"), Description);
+		bSetAField &= Json::TryGetField(InJson, TEXT("externalDocs"), ExternalDocs);
+		bSetAField &= Json::TryGetField(InJson, TEXT("operationId"), OperationId);
+		bSetAField &= Json::TryGetField(InJson, TEXT("consumes"), Consumes);
+		bSetAField &= Json::TryGetField(InJson, TEXT("produces"), Produces);
+		bSetAField &= Json::TryGetField(InJson, TEXT("parameters"), Parameters);
+		bSetAField &= Json::TryGetField(InJson, TEXT("responses"), Responses);
+		bSetAField &= Json::TryGetField(InJson, TEXT("schemes"), Schemes); // @todo
+		bSetAField &= Json::TryGetField(InJson, TEXT("security"), Security);
+		return bSetAField;
+	}
+
+	bool FReference::FromJson(const TSharedRef<FJsonObject>& InJson)
+	{
+		bool bSetAField = false;
+		bSetAField &= Json::TryGetField(InJson, TEXT("$ref"), $Ref);
+		return bSetAField;
+	}
+}
