WebAPI missing files pt2

diff --git a/Engine/Plugins/Web/WebAPI/Source/WebAPIOpenAPI/Private/V3/WebAPIOpenAPIConverter.cpp b/Engine/Plugins/Web/WebAPI/Source/WebAPIOpenAPI/Private/V3/WebAPIOpenAPIConverter.cpp
new file mode 100644
--- /dev/null
+++ b/Engine/Plugins/Web/WebAPI/Source/WebAPIOpenAPI/Private/V3/WebAPIOpenAPIConverter.cpp
@@ -0,0 +1,1909 @@
+﻿// Copyright Epic Games, Inc. All Rights Reserved.
+
+#include "WebAPIOpenAPIConverter.h"
+
+#include "IWebAPIEditorModule.h"
+#include "WebAPIDefinition.h"
+#include "WebAPITypes.h"
+#include "Algo/ForEach.h"
+#include "Dom/WebAPIEnum.h"
+#include "Dom/WebAPIModel.h"
+#include "Dom/WebAPIOperation.h"
+#include "Dom/WebAPIParameter.h"
+#include "Dom/WebAPISchema.h"
+#include "Dom/WebAPIService.h"
+#include "Dom/WebAPITypeRegistry.h"
+#include "Internationalization/BreakIterator.h"
+#include "V3/WebAPIOpenAPIProvider.h"
+#include "V3/WebAPIOpenAPISchema.h"
+
+// @todo: remove
+PRAGMA_DISABLE_OPTIMIZATION
+
+#define LOCTEXT_NAMESPACE "WebAPIOpenAPIConverter"
+
+#define SET_OPTIONAL(SrcProperty, DstProperty)			\
+if(SrcProperty.IsSet())									\
+{														\
+	DstProperty = SrcProperty.GetValue();				\
+}
+
+#define SET_OPTIONAL_FLAGGED(SrcProperty, DstProperty, DstFlag)	\
+if(SrcProperty.IsSet())											\
+{																\
+	DstFlag = true;												\
+	DstProperty = SrcProperty.GetValue();						\
+}
+
+namespace UE
+{
+	namespace WebAPI
+	{
+		// @todo: remove
+		// @see: https://blogs.sap.com/2018/01/05/open-api-spec-2.0-vs-3.0/
+		// [REDACTED_PROFILE_URL]/openapi-codegen/blob/fcbbee8ce4304400c7d2c605fb0388d6c6e5b64a/lib/orange/downconvert.js
+		namespace OpenAPI
+		{
+			FWebAPIOpenAPISchemaConverter::FWebAPIOpenAPISchemaConverter(
+				const TSharedPtr<const UE::WebAPI::OpenAPI::V3::FOpenAPIObject>& InOpenAPI, UWebAPISchema* InWebAPISchema,
+				const TSharedRef<FWebAPIMessageLog>& InMessageLog, const FWebAPIProviderSettings& InProviderSettings):
+				InputSchema(InOpenAPI)
+				, OutputSchema(InWebAPISchema)
+				, MessageLog(InMessageLog)
+				, ProviderSettings(InProviderSettings)
+			{
+			}
+
+			bool FWebAPIOpenAPISchemaConverter::Convert()
+			{
+				if(InputSchema->Info.IsValid())
+				{
+					OutputSchema->APIName = InputSchema->Info->Title;
+					OutputSchema->Version = InputSchema->Info->Version;
+				}
+
+				// V3 has multiple servers vs singular in V2
+				if(InputSchema->Servers.IsSet()
+					&& !InputSchema->Servers->IsEmpty()
+					&& InputSchema->Servers.GetValue()[0].IsValid())
+				{
+					const TSharedPtr<V3::FServerObject> Server = InputSchema->Servers.GetValue()[0];
+					
+					FString Url = Server->Url;
+					FString Scheme;
+					FParse::SchemeNameFromURI(*Url, Scheme);
+
+					// If Url isn't complete (relative, etc.), it won't have a scheme.
+					if(!Scheme.IsEmpty())
+					{
+						OutputSchema->URISchemes.Add(Scheme);
+					}
+
+					Url = Url.Replace(*(Scheme + TEXT("://")), TEXT(""));					
+					Url.Split(TEXT("/"), &OutputSchema->Host, &OutputSchema->BaseUrl);
+
+					// If Url isn't complete, there may not be a valid host
+					if(OutputSchema->Host.IsEmpty())
+					{
+						MessageLog->LogWarning(LOCTEXT("NoHostProvided", "The specification did not contain a host Url, this should be specified manually in the generated project settings."), FWebAPIOpenAPIProvider::LogName);
+					}
+				}
+
+				// If no Url schemes provided, add https as default
+				if(OutputSchema->URISchemes.IsEmpty())
+				{
+					OutputSchema->URISchemes = { TEXT("https") };
+				}
+				
+				// Top level decl of tags optional, so find from paths
+				bool bSuccessfullyConverted = InputSchema->Tags.IsSet()
+											? ConvertTags(InputSchema->Tags.GetValue(), OutputSchema.Get())
+											: true;
+
+				bSuccessfullyConverted &= InputSchema->Components.IsValid()
+										  ? ConvertModels(InputSchema->Components->Schemas, OutputSchema.Get())
+										  : true;
+				//bSuccessfullyConverted &= ConvertParameters(InputSchema->Components.Parameters, OutputSchema.Get());
+				bSuccessfullyConverted &= ConvertPaths(InputSchema->Paths, OutputSchema.Get());
+
+				return bSuccessfullyConverted;
+			}
+
+			FString FWebAPIOpenAPISchemaConverter::NameTransformer(const FWebAPINameVariant& InString) const
+			{
+				return ProviderSettings.ToPascalCase(InString);
+			}
+
+			TObjectPtr<UWebAPITypeInfo> FWebAPIOpenAPISchemaConverter::ResolveMappedType(const FString& InType)
+			{
+				const TObjectPtr<UWebAPIStaticTypeRegistry> StaticTypeRegistry = IWebAPIEditorModuleInterface::Get().GetStaticTypeRegistry();
+
+				// [REDACTED_PROFILE_URL]/openapi-generator/blob/5d68bd6a03f0c48e838b4fe3b98b7e30858c0373/modules/openapi-generator/src/main/java/org/openapitools/codegen/languages/CppUE4ClientCodegen.java
+				// OpenAPI type to UE type (Prefix, Name)
+				static TMap<FString, TObjectPtr<UWebAPITypeInfo>> TypeMap =
+				{
+					{TEXT("file"), StaticTypeRegistry->FilePath},
+					{TEXT("any"), StaticTypeRegistry->Object},
+					{TEXT("object"), StaticTypeRegistry->Object},
+
+					{TEXT("array"), StaticTypeRegistry->String},
+					{TEXT("boolean"), StaticTypeRegistry->Boolean},
+					{TEXT("byte"), StaticTypeRegistry->Byte},
+					{TEXT("integer"), StaticTypeRegistry->Int32},
+					{TEXT("int"), StaticTypeRegistry->Int32},
+					{TEXT("int32"), StaticTypeRegistry->Int32},
+					{TEXT("short"), StaticTypeRegistry->Int16},
+					{TEXT("int16"), StaticTypeRegistry->Int16},
+					{TEXT("long"), StaticTypeRegistry->Int64},
+					{TEXT("int64"), StaticTypeRegistry->Int64},
+					{TEXT("float"), StaticTypeRegistry->Float},
+					{TEXT("double"), StaticTypeRegistry->Double},
+
+					{TEXT("number"), StaticTypeRegistry->Int32},
+					{TEXT("char"), StaticTypeRegistry->Char},
+					{TEXT("date"), StaticTypeRegistry->DateTime},
+					{TEXT("date-time"), StaticTypeRegistry->DateTime},
+					{TEXT("password"), StaticTypeRegistry->String},
+					{TEXT("string"), StaticTypeRegistry->String},
+					{TEXT("void"), StaticTypeRegistry->Void},
+					{TEXT("null"), StaticTypeRegistry->Nullptr}
+				};
+
+				if (const TObjectPtr<UWebAPITypeInfo>* FoundTypeInfo = TypeMap.Find(InType))
+				{
+					return *FoundTypeInfo;
+				}
+
+				return nullptr;
+			}
+
+			TObjectPtr<UWebAPITypeInfo> FWebAPIOpenAPISchemaConverter::ResolveType(
+				FString InType,
+				FString InFormat,
+				FString InDefinitionName,
+				const TSharedPtr<OpenAPI::V3::FSchemaObject>& InSchema)
+			{
+				const TObjectPtr<UWebAPIStaticTypeRegistry> StaticTypeRegistry = IWebAPIEditorModuleInterface::Get().GetStaticTypeRegistry();
+				
+				static TMap<EJson, FString> EJson_ValueToTypeName = {
+					{EJson::None, TEXT("None")},
+					{EJson::Null, TEXT("nullptr")},
+					{EJson::String, TEXT("FString")},
+					{EJson::Boolean, TEXT("bool")},
+					{EJson::Number, TEXT("int32")},
+					{EJson::Array, TEXT("TArray")},
+					{EJson::Object, TEXT("UObject")},
+				};
+
+				static TMap<EJson, TObjectPtr<UWebAPITypeInfo>> EJson_ValueToTypeInfo = {
+					{EJson::None, StaticTypeRegistry->Void},
+					{EJson::Null, StaticTypeRegistry->Nullptr},
+					{EJson::String, StaticTypeRegistry->String},
+					{EJson::Boolean, StaticTypeRegistry->Boolean},
+					{EJson::Number, StaticTypeRegistry->Int32},
+					{EJson::Array, StaticTypeRegistry->String},
+					{EJson::Object, StaticTypeRegistry->Object},
+				};
+
+				static TMap<FString, TObjectPtr<UWebAPITypeInfo>> JsonTypeValueToTypeInfo = {
+					{TEXT("None"), StaticTypeRegistry->Void},
+					{TEXT("Null"), StaticTypeRegistry->Nullptr},
+					{TEXT("String"), StaticTypeRegistry->String},
+					{TEXT("Boolean"), StaticTypeRegistry->Boolean},
+					{TEXT("Number"), StaticTypeRegistry->Int32},
+					{TEXT("Array"), StaticTypeRegistry->String},
+					{TEXT("Object"), StaticTypeRegistry->Object},
+				};
+
+				TObjectPtr<UWebAPITypeInfo> Result = nullptr;
+
+				// If a definition name is supplied, try to find it first
+				if ((InType == TEXT("Object") || InType == TEXT("Array")) && !InDefinitionName.IsEmpty())
+				{
+					if (const TObjectPtr<UWebAPITypeInfo>* FoundTypeInfo = OutputSchema->TypeRegistry->FindGeneratedType(EWebAPISchemaType::Model, InDefinitionName))
+					{
+						Result = *FoundTypeInfo;
+					}
+				}
+
+				// If not found above
+				if (!Result)
+				{
+					// Try specific types
+					if (const TObjectPtr<UWebAPITypeInfo>& FoundMappedTypeInfo = ResolveMappedType(InFormat.IsEmpty() ? InType : InFormat))
+					{
+						return FoundMappedTypeInfo;
+					}
+					// Fallback to basic types
+					else if (const TObjectPtr<UWebAPITypeInfo>* FoundTypeInfo = JsonTypeValueToTypeInfo.Find(InType))
+					{
+						Result = *FoundTypeInfo;
+					}
+
+					// If it's not a built-in type
+					if (Result != nullptr && !Result->bIsBuiltinType)
+					{
+						// Duplicate it with the provided definition name 
+						if (!InDefinitionName.IsEmpty())
+						{
+							// @todo: what if it's an enum? how do you know at this point if it's an enum or struct?
+							// Allow prefix to be set later depending on this?
+							Result = OutputSchema->TypeRegistry->GetOrMakeGeneratedType(EWebAPISchemaType::Model,
+								InDefinitionName,
+								InDefinitionName,
+								Result);
+							Result->Prefix = TEXT("F");
+							ensure(!Result->Name.IsEmpty());
+						}
+						// otherwise duplicate it as an unnamed/partially resolved type 
+						else
+						{
+							// @todo: prevent!
+							Result = Result->Duplicate(OutputSchema->TypeRegistry);
+						}
+					}
+				}
+
+				return Result;
+			}
+
+			TObjectPtr<UWebAPITypeInfo> FWebAPIOpenAPISchemaConverter::GetTypeForContentType(const FString& InContentType)
+			{
+				const TObjectPtr<UWebAPIStaticTypeRegistry> StaticTypeRegistry = IWebAPIEditorModuleInterface::Get().GetStaticTypeRegistry();
+				
+				static TMap<FString, TObjectPtr<UWebAPITypeInfo>> ContentTypeToTypeInfo = {
+					{TEXT("application/json"), StaticTypeRegistry->JsonObject},
+					{TEXT("application/xml"), StaticTypeRegistry->String},
+					{TEXT("text/plain"), StaticTypeRegistry->String}
+				};
+				
+				if(const TObjectPtr<UWebAPITypeInfo>* FoundTypeInfo = ContentTypeToTypeInfo.Find(InContentType.ToLower()))
+				{
+					return *FoundTypeInfo;
+				}
+
+				// For all other cases, use string
+				return StaticTypeRegistry->String;
+			}
+
+			template <>
+			FString FWebAPIOpenAPISchemaConverter::GetDefaultJsonTypeForStorage<EWebAPIParameterStorage>(const EWebAPIParameterStorage& InStorage)
+			{
+				static TMap<EWebAPIParameterStorage, FString> StorageToJsonType = {
+					{ EWebAPIParameterStorage::Body, TEXT("object") },
+					{ EWebAPIParameterStorage::Cookie, TEXT("string") },
+					{ EWebAPIParameterStorage::Header, TEXT("string") },
+					{ EWebAPIParameterStorage::Path, TEXT("string") },
+					{ EWebAPIParameterStorage::Query, TEXT("string") },
+				};
+
+				return StorageToJsonType[InStorage];
+			}
+
+			template <>
+			FString FWebAPIOpenAPISchemaConverter::GetDefaultJsonTypeForStorage<EWebAPIResponseStorage>(const EWebAPIResponseStorage& InStorage)
+			{
+				static TMap<EWebAPIResponseStorage, FString> StorageToJsonType = {
+					{ EWebAPIResponseStorage::Body, TEXT("object") },
+					{ EWebAPIResponseStorage::Header, TEXT("string") },
+				};
+
+				return StorageToJsonType[InStorage];
+			}
+
+			template <typename SchemaType>
+			TObjectPtr<UWebAPITypeInfo> FWebAPIOpenAPISchemaConverter::ResolveType(const TSharedPtr<SchemaType>& InSchema, const FString& InDefinitionName, FString InJsonType)
+			{
+				static_assert(TIsDerivedFrom<SchemaType, OpenAPI::V3::FSchemaObjectBase>::Value, "Type is not derived from OpenAPI::V3::FSchemaObjectBase.");
+
+				const TObjectPtr<UWebAPIStaticTypeRegistry> StaticTypeRegistry = IWebAPIEditorModuleInterface::Get().GetStaticTypeRegistry();
+
+				FString DefinitionName = InDefinitionName;
+				TObjectPtr<UWebAPITypeInfo> Result = nullptr;
+				
+				TSharedPtr<UE::WebAPI::OpenAPI::V3::FSchemaObject> ItemSchema = nullptr;
+				if(InSchema.IsValid() && InSchema->Items.IsSet())
+				{
+					ItemSchema = ResolveReference(InSchema->Items.GetValue(), DefinitionName);
+					if(!InSchema->Items->GetPath().IsEmpty())
+					{
+						DefinitionName = InSchema->Items->GetLastPathSegment();						
+					}
+					return ResolveType(ItemSchema, DefinitionName);
+				}
+
+				if(InSchema.IsValid())
+				{
+					Result = ResolveType(InSchema->Type.Get(InJsonType.IsEmpty() ? TEXT("object") : InJsonType), InSchema->Format.Get(TEXT("")), DefinitionName);
+				}
+				else
+				{
+					Result = ResolveType(InJsonType, TEXT(""), DefinitionName);
+					if(!Result)
+					{
+						Result = StaticTypeRegistry->Object;
+					}
+				}
+
+				if (InSchema.IsValid() && InSchema->Enum.IsSet() && !InSchema->Enum.GetValue().IsEmpty())
+				{
+					if(const TObjectPtr<UWebAPITypeInfo>* FoundGeneratedType = OutputSchema->TypeRegistry->FindGeneratedType(EWebAPISchemaType::Model, NameTransformer(DefinitionName)))
+					{
+						Result = *FoundGeneratedType;						
+					}
+					else
+					{
+						Result = OutputSchema->TypeRegistry->GetOrMakeGeneratedType(
+							EWebAPISchemaType::Model,
+							NameTransformer(DefinitionName),
+							DefinitionName,
+							StaticTypeRegistry->Enum);
+					}
+				}
+				else if (!Result->IsEnum() && (Result == StaticTypeRegistry->Object || Result->ToString(true).IsEmpty()))
+				{
+					if (!DefinitionName.IsEmpty())
+					{
+						if (const TObjectPtr<UWebAPITypeInfo>* FoundBuiltinType = StaticTypeRegistry->FindBuiltinType(DefinitionName))
+						{
+							Result = *FoundBuiltinType;
+						}
+						else if (const TObjectPtr<UWebAPITypeInfo>* FoundGeneratedModelType = OutputSchema->TypeRegistry->FindGeneratedType(EWebAPISchemaType::Model, DefinitionName))
+						{
+							Result = *FoundGeneratedModelType;
+						}
+						else if (const TObjectPtr<UWebAPITypeInfo>* FoundGeneratedParameterType = OutputSchema->TypeRegistry->FindGeneratedType(EWebAPISchemaType::Parameter, DefinitionName))
+						{
+							Result = *FoundGeneratedParameterType;
+						}
+						else
+						{
+							FFormatNamedArguments Args;
+							Args.Add(TEXT("DefinitionName"), FText::FromString(DefinitionName));
+							MessageLog->LogInfo(FText::Format(LOCTEXT("CannotResolveType", "ResolveType (object) failed to find a matching type for definition \"{DefinitionName}\", creating a new one."), Args), FWebAPIOpenAPIProvider::LogName);
+
+							// @bug: where is the model created for this?
+							// @todo: should this always be of type model, vs parameter?
+							Result = OutputSchema->TypeRegistry->GetOrMakeGeneratedType(
+								EWebAPISchemaType::Model,
+								NameTransformer(DefinitionName),
+								{},
+								TEXT("F"));
+							Result->JsonType = UWebAPIStaticTypeRegistry::ToFromJsonType;
+						}
+					}
+				}
+
+				check(!Result.IsNull());
+				
+				return Result;
+			}
+
+			template <>
+			TSharedPtr<OpenAPI::V3::FSchemaObject> FWebAPIOpenAPISchemaConverter::ResolveReference<OpenAPI::V3::FSchemaObject>(const FString& InDefinitionName)
+			{
+				if(!InputSchema->Components.IsValid())
+				{
+					return nullptr;
+				}
+				
+				if (const Json::TJsonReference<OpenAPI::V3::FSchemaObject>* FoundDefinition = InputSchema->Components->Schemas.Find(InDefinitionName))
+				{
+					if(!FoundDefinition->IsSet())
+					{
+						return ResolveReference<OpenAPI::V3::FSchemaObject>(FoundDefinition->GetLastPathSegment());
+					}
+					return FoundDefinition->GetShared();
+				}
+
+				return nullptr;
+			}
+
+			template <>
+			TSharedPtr<OpenAPI::V3::FParameterObject> FWebAPIOpenAPISchemaConverter::ResolveReference<OpenAPI::V3::FParameterObject>(const FString& InDefinitionName)
+			{
+				if(!InputSchema->Components.IsValid())
+				{
+					return nullptr;
+				}
+				
+				if (const Json::TJsonReference<OpenAPI::V3::FParameterObject>* FoundDefinition = InputSchema->Components->Parameters.Find(InDefinitionName))
+				{
+					if(!FoundDefinition->IsSet())
+					{
+						return ResolveReference<OpenAPI::V3::FParameterObject>(FoundDefinition->GetLastPathSegment());
+					}
+					return FoundDefinition->GetShared();
+				}
+
+				return nullptr;
+			}
+
+			template <>
+			TSharedPtr<OpenAPI::V3::FResponseObject> FWebAPIOpenAPISchemaConverter::ResolveReference<V3::FResponseObject>(const FString& InDefinitionName)
+			{
+				if(!InputSchema->Components.IsValid())
+				{
+					return nullptr;
+				}
+				
+				if (const Json::TJsonReference<OpenAPI::V3::FResponseObject>* FoundDefinition = InputSchema->Components->Responses.Find(InDefinitionName))
+				{
+					if(!FoundDefinition->IsSet())
+					{
+						return ResolveReference<OpenAPI::V3::FResponseObject>(FoundDefinition->GetLastPathSegment());
+					}
+					return FoundDefinition->GetShared();
+				}
+
+				return nullptr;
+			}
+
+			template <>
+			TSharedPtr<OpenAPI::V3::FSecuritySchemeObject> FWebAPIOpenAPISchemaConverter::ResolveReference<V3::FSecuritySchemeObject>(const FString& InDefinitionName)
+			{
+				if(!InputSchema->Components.IsValid())
+				{
+					return nullptr;
+				}
+				
+				if (const Json::TJsonReference<OpenAPI::V3::FSecuritySchemeObject>* FoundDefinition = InputSchema->Components->SecuritySchemes.Find(InDefinitionName))
+				{
+					if(!FoundDefinition->IsSet())
+					{
+						return ResolveReference<OpenAPI::V3::FSecuritySchemeObject>(FoundDefinition->GetLastPathSegment());
+					}
+					return FoundDefinition->GetShared();
+				}
+
+				return nullptr;
+			}
+
+			template <>
+			TSharedPtr<OpenAPI::V3::FRequestBodyObject> FWebAPIOpenAPISchemaConverter::ResolveReference<V3::FRequestBodyObject>(const FString& InDefinitionName)
+			{
+				if (const Json::TJsonReference<OpenAPI::V3::FRequestBodyObject>* FoundDefinition = InputSchema->Components->RequestBodies.Find(InDefinitionName))
+				{
+					if(!FoundDefinition->IsSet())
+					{
+						return ResolveReference<OpenAPI::V3::FRequestBodyObject>(FoundDefinition->GetLastPathSegment());
+					}
+					return FoundDefinition->GetShared();
+				}
+
+				return nullptr;
+			}
+
+			FWebAPINameVariant FWebAPIOpenAPISchemaConverter::ResolvePropertyName(const TObjectPtr<UWebAPIProperty>& InProperty, const FWebAPITypeNameVariant& InPotentialName, const TOptional<bool>& bInIsArray)
+			{
+				check(InProperty);
+
+				const bool bIsArray = bInIsArray.Get(InProperty->bIsArray);
+
+				// If it's an array, it may be called the generic "Values", try to find a better name
+				if(bIsArray && InProperty->Name == ProviderSettings.GetDefaultArrayPropertyName())
+				{
+					if(InPotentialName.IsValid())
+					{
+						FWebAPINameInfo NameInfo = FWebAPINameInfo(ProviderSettings.Pluralize(InPotentialName.ToString(true)), InProperty->Name.GetJsonName());
+						return NameInfo;
+					}
+					
+					if(InProperty->Type.HasTypeInfo() && !InProperty->Type.TypeInfo->bIsBuiltinType)
+					{
+						FWebAPINameInfo NameInfo = FWebAPINameInfo(ProviderSettings.Pluralize(InProperty->Type.ToString(true)), InProperty->Name.GetJsonName());
+						return NameInfo;
+					}
+				}
+				else if(InProperty->Name == ProviderSettings.GetDefaultPropertyName())
+				{
+					if(InPotentialName.IsValid())
+					{
+						FWebAPINameInfo NameInfo = FWebAPINameInfo(InPotentialName.ToString(true), InProperty->Name.GetJsonName());
+						return NameInfo;
+					}
+					
+					if(InProperty->Type.HasTypeInfo() && !InProperty->Type.TypeInfo->bIsBuiltinType)
+					{
+						FWebAPINameInfo NameInfo = FWebAPINameInfo(InProperty->Type.ToString(true), InProperty->Name.GetJsonName());
+						return NameInfo;
+					}
+				}
+
+				return InProperty->Name;
+			}
+
+			template <typename ObjectType>
+			TSharedPtr<ObjectType> FWebAPIOpenAPISchemaConverter::ResolveReference(const Json::TJsonReference<ObjectType>& InJsonReference, FString& OutDefinitionName, bool bInCheck)
+			{
+				if (InJsonReference.IsSet())
+				{
+					return InJsonReference.GetShared();
+				}
+
+				if(!InJsonReference.IsValid())
+				{
+					return nullptr;
+				}
+				
+				FString DefinitionName = InJsonReference.GetLastPathSegment();
+				if (TSharedPtr<ObjectType> FoundDefinition = ResolveReference<ObjectType>(DefinitionName))
+				{
+					OutDefinitionName = DefinitionName; // Only set if found 
+					return FoundDefinition;
+				}
+
+				if(bInCheck)
+				{
+					FFormatNamedArguments Args;
+					Args.Add(TEXT("ReferencePath"), FText::FromString(InJsonReference.GetPath()));
+					MessageLog->LogWarning(FText::Format(LOCTEXT("CannotResolveJsonReference", "Couldn't resolve JsonReference \"{ReferencePath}\""), Args),	FWebAPIOpenAPIProvider::LogName);
+				}
+
+				return nullptr;
+			}
+
+			template <typename ObjectType>
+			TSharedPtr<ObjectType> FWebAPIOpenAPISchemaConverter::ResolveReference(const UE::Json::TJsonReference<ObjectType>& InJsonReference, bool bInCheck)
+			{
+				FString DefinitionName;
+				return ResolveReference(InJsonReference, DefinitionName, bInCheck);
+			}
+
+			template <>
+			bool FWebAPIOpenAPISchemaConverter::IsArray(const TSharedPtr<UE::WebAPI::OpenAPI::V3::FSchemaObject>& InSchema)
+			{
+				return InSchema->Type.Get(TEXT(""))	== TEXT("array");
+			}
+
+			template <>
+			bool FWebAPIOpenAPISchemaConverter::IsArray(const TSharedPtr<UE::WebAPI::OpenAPI::V3::FParameterObject>& InSchema)
+			{
+				FString FoundDefinitionName;
+				return InSchema->Schema.IsSet()
+					? IsArray(ResolveReference(InSchema->Schema.GetValue(), FoundDefinitionName))
+					: false;
+			}
+
+			template <typename SchemaType, typename ModelType>
+			bool FWebAPIOpenAPISchemaConverter::ConvertModelBase(const TSharedPtr<SchemaType>& InSchema,
+																const TObjectPtr<ModelType>& OutModel)
+			{
+				static_assert(TIsDerivedFrom<SchemaType, OpenAPI::V3::FSchemaObjectBase>::Value, "Type is not derived from OpenAPI::V3::FSchemaObjectBase.");
+/*
+				OutModel->Description = InSchema->Description.Get(TEXT(""));
+
+				SET_OPTIONAL(InSchema->bRequired, OutModel->bIsRequired)
+				
+				SET_OPTIONAL_FLAGGED(InSchema->Minimum, OutModel->MinimumValue, OutModel->bUseMinimumValue);
+				SET_OPTIONAL_FLAGGED(InSchema->Maximum, OutModel->MaximumValue, OutModel->bUseMaximumValue);
+
+				if(!OutModel->bUseMinimumValue)
+				{
+					SET_OPTIONAL_FLAGGED(InSchema->MinLength, OutModel->MinimumValue, OutModel->bUseMinimumValue);
+				}
+
+				if(!OutModel->bUseMaximumValue)
+				{
+					SET_OPTIONAL_FLAGGED(InSchema->MaxLength, OutModel->MaximumValue, OutModel->bUseMaximumValue);					
+				}
+				
+				SET_OPTIONAL_FLAGGED(InSchema->Pattern, OutModel->Pattern, OutModel->bUsePattern);
+				*/
+				return true;
+			}
+
+			template <typename SchemaType>
+			TObjectPtr<UWebAPIEnum> FWebAPIOpenAPISchemaConverter::ConvertEnum(const TSharedPtr<SchemaType>& InSrcSchema, const FWebAPITypeNameVariant& InEnumTypeName) const
+			{
+				static_assert(TIsDerivedFrom<SchemaType, OpenAPI::V3::FSchemaObjectBase>::Value, "Type is not derived from OpenAPI::V3::FSchemaObjectBase.");
+
+				FWebAPITypeNameVariant EnumTypeName;
+				if(InEnumTypeName.IsValid())
+				{
+					EnumTypeName = InEnumTypeName;
+				}
+				else
+				{
+					const FString EnumName = InSrcSchema->Name;
+					check(!EnumName.IsEmpty());
+
+					const TObjectPtr<UWebAPIStaticTypeRegistry> StaticTypeRegistry = IWebAPIEditorModuleInterface::Get().GetStaticTypeRegistry();
+					EnumTypeName = OutputSchema->TypeRegistry->GetOrMakeGeneratedType(
+						EWebAPISchemaType::Model,
+						NameTransformer(EnumName),
+						EnumName,
+						StaticTypeRegistry->Enum);
+				}
+
+				const TObjectPtr<UWebAPIEnum>& DstEnum = OutputSchema->AddEnum(EnumTypeName.TypeInfo.Get());
+				DstEnum->Name = EnumTypeName;
+				DstEnum->Description = InSrcSchema->Description.Get(TEXT(""));
+
+				const TObjectPtr<UWebAPIModelBase> ModelBase = Cast<UWebAPIModelBase>(DstEnum);
+				if (!ConvertModelBase(InSrcSchema, ModelBase))
+				{
+					return nullptr;
+				}
+
+				for (const FString& SrcEnumValue : InSrcSchema->Enum.GetValue())
+				{
+					const TObjectPtr<UWebAPIEnumValue> DstEnumValue = DstEnum->AddValue();
+					DstEnumValue->Name.NameInfo.Name = NameTransformer(SrcEnumValue);
+					DstEnumValue->Name.NameInfo.JsonName = SrcEnumValue;
+				}
+
+				DstEnum->BindToTypeInfo();
+				
+				return DstEnum;
+			}
+
+			template <>
+			bool FWebAPIOpenAPISchemaConverter::ConvertProperty<OpenAPI::V3::FParameterObject>(
+				const FWebAPITypeNameVariant& InModelName,
+				const FWebAPINameVariant& InPropertyName,
+				const TSharedPtr<OpenAPI::V3::FParameterObject>& InParameter,
+				const FString& InDefinitionName,
+				const TObjectPtr<UWebAPIProperty>& OutProperty)
+			{
+				const FString DefinitionName = !InDefinitionName.IsEmpty()
+												? InDefinitionName
+												: ProviderSettings.MakeNestedPropertyTypeName(InModelName, InPropertyName);
+
+				const TObjectPtr<UWebAPIModelBase> ModelBase = OutProperty;
+				if (!ConvertModelBase(InParameter, ModelBase))
+				{
+					return false;
+				}
+
+				OutProperty->bIsArray = IsArray(InParameter);
+				
+				OutProperty->Name = FWebAPINameInfo(NameTransformer(InPropertyName.ToString()), InPropertyName.GetJsonName(), OutProperty->Type);				
+				OutProperty->Name = ResolvePropertyName(OutProperty, InModelName, {});
+				OutProperty->bIsRequired = InParameter->bRequired.Get(false);
+				OutProperty->BindToTypeInfo();
+
+				if(InParameter->Schema.IsSet())
+				{
+					const TSharedPtr<V3::FSchemaObject> Schema = ResolveReference(InParameter->Schema.GetValue());
+					OutProperty->Type = ResolveType(Schema);
+					
+					// Add enum as it's own model, and reference it as this properties type
+					if (InParameter->Schema.GetValue()->Enum.IsSet() && !InParameter->Schema.GetValue()->Enum->IsEmpty())
+					{
+						const FWebAPITypeNameVariant EnumTypeName = OutProperty->Type;
+						EnumTypeName.TypeInfo->SetNested(InModelName);
+						FString EnumName = OutProperty->Name.ToString(true);
+
+						// Only make nested name if the model and property name aren't the same, otherwise you get "NameName"!
+						if(EnumName != InModelName)
+						{
+							EnumName = ProviderSettings.MakeNestedPropertyTypeName(InModelName, OutProperty->Name.ToString(true));
+						}
+						EnumTypeName.TypeInfo->SetName(EnumName);
+
+						const TObjectPtr<UWebAPIEnum>& Enum = ConvertEnum(Schema, EnumTypeName);
+
+						const FText LogMessage = FText::FormatNamed(
+							LOCTEXT("AddedImplicitEnumForPropertyOfModel", "Implicit enum created for property \"{PropertyName}\" of model \"{ModelName}\"."),
+							TEXT("PropertyName"), FText::FromString(*InPropertyName.ToString(true)),
+							TEXT("ModelName"), FText::FromString(*InModelName.ToString(true)));
+						MessageLog->LogInfo(LogMessage, FWebAPIOpenAPIProvider::LogName);
+					
+						Enum->Name.TypeInfo->DebugString += LogMessage.ToString();
+						Enum->Name.TypeInfo->JsonName = InPropertyName.GetJsonName();
+
+						OutProperty->Type = Enum->Name;
+						OutProperty->Type.TypeInfo->Model = Enum;
+					}
+				}
+
+				return true;
+			}
+
+			template <typename SchemaType>
+			bool FWebAPIOpenAPISchemaConverter::ConvertProperty(const FWebAPITypeNameVariant& InModelName,
+																const FWebAPINameVariant& InPropertyName,
+																const TSharedPtr<SchemaType>& InSchema,
+																const FString& InDefinitionName,
+																const TObjectPtr<UWebAPIProperty>& OutProperty)
+			{
+				static_assert(TIsDerivedFrom<SchemaType, OpenAPI::V3::FSchemaObjectBase>::Value, "Type is not derived from OpenAPI::V3::FSchemaObjectBase.");
+
+				FString DefinitionName = !InDefinitionName.IsEmpty()
+										? InDefinitionName
+										: ProviderSettings.MakeNestedPropertyTypeName(InModelName, InPropertyName);
+
+				TSharedPtr<UE::WebAPI::OpenAPI::V3::FSchemaObject> ItemSchema = InSchema;
+				if(InSchema->Items.IsSet())
+				{
+					if(!InSchema->Items->GetPath().IsEmpty())
+					{
+						ItemSchema = ResolveReference(InSchema->Items.GetValue(), DefinitionName);
+						DefinitionName = InSchema->Items->GetLastPathSegment();						
+					}
+				}
+				else
+				{
+					ItemSchema = InSchema;
+				}
+
+				OutProperty->Type = ResolveType<SchemaType>(ItemSchema, DefinitionName);
+				check(OutProperty->Type.IsValid());
+
+				const TObjectPtr<UWebAPIModelBase> ModelBase = OutProperty;
+				if (!ConvertModelBase(ItemSchema, ModelBase))
+				{
+					return false;
+				}
+
+				OutProperty->bIsArray = IsArray(InSchema);
+				
+				OutProperty->Name = FWebAPINameInfo(NameTransformer(InPropertyName.ToString()), InPropertyName.GetJsonName(), OutProperty->Type);
+				OutProperty->Name = ResolvePropertyName(OutProperty, OutProperty->Type.ToString(true), {});
+				OutProperty->bIsRequired = InSchema->bRequired.Get(false);
+				OutProperty->BindToTypeInfo();
+
+				// Add enum as it's own model, and reference it as this properties type
+				if (InSchema->Enum.IsSet() && !InSchema->Enum->IsEmpty())
+				{
+					const FWebAPITypeNameVariant EnumTypeName = OutProperty->Type;
+					EnumTypeName.TypeInfo->SetName(ProviderSettings.MakeNestedPropertyTypeName(InModelName, InPropertyName));
+
+					const TObjectPtr<UWebAPIEnum>& Enum = ConvertEnum(InSchema, EnumTypeName);
+
+					const FText LogMessage = FText::FormatNamed(
+						LOCTEXT("AddedImplicitEnumForPropertyOfModel", "Implicit enum created for property \"{PropertyName}\" of model \"{ModelName}\"."),
+						TEXT("PropertyName"), FText::FromString(*InPropertyName.ToString(true)),
+						TEXT("ModelName"), FText::FromString(*InModelName.ToString(true)));
+					MessageLog->LogInfo(LogMessage, FWebAPIOpenAPIProvider::LogName);
+
+					Enum->Name.TypeInfo->DebugString += LogMessage.ToString();
+					
+					Enum->Name.TypeInfo->JsonName = InPropertyName.GetJsonName();
+					Enum->Name.TypeInfo->SetNested(InModelName);
+
+					OutProperty->Type = Enum->Name;
+					OutProperty->Type.TypeInfo->Model = Enum;
+				}
+				// Add struct as it's own model, and reference it as this properties type
+				else if (OutProperty->Type.ToString(true).IsEmpty())
+				{
+					check(OutProperty->Type.HasTypeInfo());
+					
+					const FWebAPITypeNameVariant ModelTypeName = OutProperty->Type;
+					ModelTypeName.TypeInfo->SetName(ProviderSettings.MakeNestedPropertyTypeName(InModelName, InPropertyName));
+					ModelTypeName.TypeInfo->Prefix = TEXT("F");
+					ModelTypeName.TypeInfo->SetNested(InModelName);
+
+					const TObjectPtr<UWebAPIModel>& Model = OutputSchema->AddModel<UWebAPIModel>(ModelTypeName.HasTypeInfo() ? ModelTypeName.TypeInfo.Get() : nullptr);
+					ConvertModel(InSchema, {}, Model);
+
+					const FText LogMessage = FText::FormatNamed(
+						LOCTEXT("AddedImplicitModelForPropertyOfModel", "Implicit model created for property \"{PropertyName}\" of model \"{ModelName}\"."),
+						TEXT("PropertyName"), FText::FromString(*InPropertyName.ToString(true)),
+						TEXT("ModelName"), FText::FromString(*InModelName.ToString(true)));
+					MessageLog->LogInfo(LogMessage, FWebAPIOpenAPIProvider::LogName);
+					
+					Model->Name.TypeInfo->DebugString += LogMessage.ToString();
+					
+					Model->Name.TypeInfo->JsonName = InPropertyName.GetJsonName();
+					Model->Name.TypeInfo->JsonType = UWebAPIStaticTypeRegistry::ToFromJsonType;
+
+					OutProperty->Type = Model->Name;
+					OutProperty->Type.TypeInfo->Model = Model;
+				}
+				
+				OutProperty->Name = ResolvePropertyName(OutProperty, InModelName);
+
+				return true;
+			}
+
+			bool FWebAPIOpenAPISchemaConverter::ConvertProperty(const FWebAPITypeNameVariant& InModelName, const FWebAPINameVariant& InPropertyName, const FWebAPITypeNameVariant& InPropertyTypeName, const TObjectPtr<UWebAPIProperty>& OutProperty)
+			{
+				OutProperty->Type = InPropertyTypeName;
+				check(OutProperty->Type.IsValid());
+
+				OutProperty->Name = FWebAPINameInfo(NameTransformer(InPropertyName.ToString()), InPropertyName.GetJsonName(), OutProperty->Type);
+				OutProperty->Name = ResolvePropertyName(OutProperty, OutProperty->Type.ToString(true), {});
+				OutProperty->bIsRequired = false;
+				OutProperty->BindToTypeInfo();
+
+				return true;
+			}
+
+			template <typename SchemaType>
+			TObjectPtr<UWebAPIProperty> FWebAPIOpenAPISchemaConverter::ConvertProperty(const TSharedPtr<SchemaType>& InSrcSchema, const TObjectPtr<UWebAPIModel>& InModel, const FWebAPINameVariant& InPropertyName, const FString& InDefinitionName)
+			{
+				static_assert(TIsDerivedFrom<SchemaType, OpenAPI::V3::FSchemaObjectBase>::Value, "Type is not derived from OpenAPI::V3::FSchemaObjectBase.");
+
+				const FWebAPINameVariant SrcPropertyName = InPropertyName;
+				FString SrcPropertyDefinitionName = ProviderSettings.MakeNestedPropertyTypeName(*InModel->Name.ToString(true), SrcPropertyName);
+				const TSharedPtr<OpenAPI::V3::FSchemaObject> SrcPropertyValue = ResolveReference(InSrcSchema, SrcPropertyDefinitionName);
+
+				if (SrcPropertyValue)
+				{
+					const TObjectPtr<UWebAPIProperty>& DstProperty = InModel->Properties.Add_GetRef(NewObject<UWebAPIProperty>(InModel));
+					ConvertProperty(InModel->Name,
+						FWebAPINameInfo(NameTransformer(SrcPropertyName), SrcPropertyName.GetJsonName(), InModel->Name),
+						SrcPropertyValue,
+						SrcPropertyDefinitionName,
+						DstProperty);
+					
+					return DstProperty;
+				}
+
+				return nullptr;
+			}
+
+			template <typename SchemaType>
+			bool FWebAPIOpenAPISchemaConverter::ConvertModel(const TSharedPtr<SchemaType>& InSrcSchema, const FWebAPITypeNameVariant& InModelTypeName, const TObjectPtr<UWebAPIModel>& OutModel)
+			{
+				static_assert(TIsDerivedFrom<SchemaType, OpenAPI::V3::FSchemaObjectBase>::Value, "Type is not derived from OpenAPI::V3::FSchemaObjectBase.");
+
+				FWebAPITypeNameVariant ModelTypeName;
+				if(InModelTypeName.IsValid())
+				{
+					ModelTypeName = InModelTypeName;
+				}
+				else if(OutModel->Name.HasTypeInfo())
+				{
+					ModelTypeName = OutModel->Name;
+				}
+				else
+				{
+					const FString ModelName = InSrcSchema->Name;
+					check(!ModelName.IsEmpty());
+
+					ModelTypeName = OutputSchema->TypeRegistry->GetOrMakeGeneratedType(
+						EWebAPISchemaType::Model,
+						NameTransformer(*ModelName),
+						*ModelName,
+						TEXT("F"));
+				}
+
+				if (!InSrcSchema.IsValid())
+				{
+					FFormatNamedArguments Args;
+					Args.Add(TEXT("ModelName"), FText::FromString(*ModelTypeName.ToString(true)));
+					MessageLog->LogWarning(
+						FText::Format(LOCTEXT("SchemaInvalid",
+								"The schema for model \"{ModelName}\" was invalid/null."),
+							Args),
+						FWebAPIOpenAPIProvider::LogName);
+					return false;
+				}
+
+				const TObjectPtr<UWebAPIModel>& Model = OutModel ? OutModel : OutputSchema->AddModel<UWebAPIModel>(ModelTypeName.TypeInfo.Get());
+
+				const TObjectPtr<UWebAPIModelBase> ModelBase = Model;
+				if (!ConvertModelBase(InSrcSchema, ModelBase))
+				{
+					return false;
+				}
+
+				Model->Name = ModelTypeName;
+
+				if (InSrcSchema->Properties.IsSet() && !InSrcSchema->Properties->IsEmpty())
+				{
+					for (const TTuple<FString, Json::TJsonReference<OpenAPI::V3::FSchemaObject>>& NamePropertyPair : InSrcSchema->Properties.GetValue())
+					{
+						const FString& SrcPropertyName = NamePropertyPair.Key;
+						
+						FString SrcPropertyDefinitionName = ProviderSettings.MakeNestedPropertyTypeName(*ModelTypeName.ToString(true), SrcPropertyName);
+						TSharedPtr<OpenAPI::V3::FSchemaObject> SrcPropertySchema = ResolveReference(NamePropertyPair.Value, SrcPropertyDefinitionName);
+
+						TObjectPtr<UWebAPIProperty>& DstProperty = Model->Properties.Add_GetRef(NewObject<UWebAPIProperty>(Model));
+						ConvertProperty(Model->Name,
+										FWebAPINameInfo(NameTransformer(SrcPropertyName), SrcPropertyName, ModelTypeName),
+										SrcPropertySchema,
+										SrcPropertyDefinitionName,
+										DstProperty);
+					}
+				}
+
+				Model->BindToTypeInfo();
+
+				return true;
+			}
+
+			template <>
+			bool FWebAPIOpenAPISchemaConverter::ConvertModel<OpenAPI::V3::FParameterObject>(const TSharedPtr<OpenAPI::V3::FParameterObject>& InSrcSchema, const FWebAPITypeNameVariant& InModelTypeName, const TObjectPtr<UWebAPIModel>& OutModel)
+			{
+				FWebAPITypeNameVariant ModelTypeName;
+				if(InModelTypeName.IsValid() && InModelTypeName.HasTypeInfo())
+				{
+					ModelTypeName = InModelTypeName;
+				}
+				else if(OutModel->Name.HasTypeInfo())
+				{
+					ModelTypeName = OutModel->Name;
+				}
+				else
+				{
+					const FString ModelName = InSrcSchema->Name;
+					check(!ModelName.IsEmpty());
+
+					ModelTypeName = OutputSchema->TypeRegistry->GetOrMakeGeneratedType(
+						EWebAPISchemaType::Model,
+						NameTransformer(*ModelName),
+						*ModelName,
+						TEXT("F"));
+				}
+
+				if (!InSrcSchema.IsValid())
+				{
+					FFormatNamedArguments Args;
+					Args.Add(TEXT("ModelName"), FText::FromString(*ModelTypeName.ToString(true)));
+					MessageLog->LogWarning(
+						FText::Format(LOCTEXT("SchemaInvalid",
+								"The schema for model \"{ModelName}\" was invalid/null."),
+							Args),
+						FWebAPIOpenAPIProvider::LogName);
+					return false;
+				}
+
+				const TObjectPtr<UWebAPIModel>& Model = OutModel ? OutModel : OutputSchema->AddModel<UWebAPIModel>(ModelTypeName.TypeInfo.Get());
+
+				const TObjectPtr<UWebAPIModelBase> ModelBase = Model;
+				if (!ConvertModelBase(InSrcSchema, ModelBase))
+				{
+					return false;
+				}
+
+				Model->Name = ModelTypeName;
+
+				Model->BindToTypeInfo();
+
+				return true;
+			}
+
+			template <typename SchemaType>
+			TObjectPtr<UWebAPIModel> FWebAPIOpenAPISchemaConverter::ConvertModel(const TSharedPtr<SchemaType>& InSrcSchema, const FWebAPITypeNameVariant& InModelTypeName)
+			{
+				check(InSrcSchema.IsValid());
+
+				FWebAPITypeNameVariant ModelTypeName;
+				if(InModelTypeName.IsValid() && ModelTypeName.HasTypeInfo())
+				{
+					ModelTypeName = InModelTypeName;
+				}
+				else
+				{
+					const FString ModelName = InSrcSchema->Name;
+					check(!ModelName.IsEmpty());
+
+					ModelTypeName = OutputSchema->TypeRegistry->GetOrMakeGeneratedType(
+						EWebAPISchemaType::Model,
+						NameTransformer(*ModelName),
+						*ModelName,
+						TEXT("F"));
+				}
+				
+				const TObjectPtr<UWebAPIModel>& Model = OutputSchema->AddModel<UWebAPIModel>(ModelTypeName.TypeInfo.Get());
+				if(ConvertModel(InSrcSchema, ModelTypeName, Model))
+				{
+					return Model;
+				}
+
+				return nullptr;
+			}
+			
+			bool FWebAPIOpenAPISchemaConverter::ConvertOperationParameter(
+				const FWebAPINameVariant& InParameterName,
+				const TSharedPtr<OpenAPI::V3::FParameterObject>& InParameter,
+				const FString& InDefinitionName,
+				const TObjectPtr<UWebAPIOperationParameter>& OutParameter)
+			{
+				// Will get schema or create if it doesn't exist (but will be empty)
+				const TSharedPtr<V3::FSchemaObject> ParameterSchema = ResolveReference(InParameter->Schema.Get({}));
+
+				static TMap<FString, EWebAPIParameterStorage> InToStorage = {
+					{TEXT("query"), EWebAPIParameterStorage::Query},
+					{TEXT("header"), EWebAPIParameterStorage::Header},
+					{TEXT("path"), EWebAPIParameterStorage::Path},
+					{TEXT("formData"), EWebAPIParameterStorage::Body},
+					{TEXT("body"), EWebAPIParameterStorage::Body}
+				};
+
+				OutParameter->Storage = InToStorage[InParameter->In];
+				
+				if (ParameterSchema.IsValid())
+				{
+					OutParameter->Type = ResolveType(ParameterSchema, InDefinitionName);
+				}
+				else
+				{
+					// @todo:
+					// OutParameter->Type = ResolveType(InParameter, InDefinitionName, GetDefaultJsonTypeForStorage(OutParameter->Storage));
+				}
+
+				OutParameter->bIsArray = IsArray(InParameter);
+				if (!OutParameter->Type.HasTypeInfo())
+				{
+					if(InParameter->Schema.IsSet())
+					{
+						// @todo
+						/* 
+						OutParameter->Type = ResolveType(InParameter->Schema->Get(GetDefaultJsonTypeForStorage(OutParameter->Storage)),
+							TEXT(""), //InParameter->Format.Get(TEXT("")),
+							InDefinitionName,
+							ParameterSchema.GetShared());
+							*/
+					}
+				}
+
+				const TObjectPtr<UWebAPIModelBase> ModelBase = Cast<UWebAPIModelBase>(OutParameter);
+				if (!ConvertModelBase(InParameter, ModelBase))
+				{
+					return nullptr;
+				}
+
+				if(!OutParameter->Type.TypeInfo->bIsBuiltinType)
+				{
+					// @todo: make generated type here
+					auto x = OutParameter->Type;
+					// ConvertEnum(InParameter)
+				}
+
+				OutParameter->Name = FWebAPINameInfo(InParameterName.ToString(true), InParameterName.GetJsonName(), OutParameter->Type);
+
+				// Add struct as it's own model, and reference it as this properties type
+				if (OutParameter->Type.ToString(true).IsEmpty())
+				{
+					const FWebAPITypeNameVariant ModelTypeName = OutParameter->Type;
+					ModelTypeName.TypeInfo->SetName(InDefinitionName);
+					ModelTypeName.TypeInfo->Prefix = TEXT("F");
+
+					const TObjectPtr<UWebAPIModel>& Model = OutputSchema->AddModel<UWebAPIModel>(
+						ModelTypeName.HasTypeInfo() ? ModelTypeName.TypeInfo.Get() : nullptr);
+					ConvertModel(InParameter->Schema.GetValue().GetShared(), {}, Model);
+
+					const FText LogMessage = FText::FormatNamed(
+						LOCTEXT("AddedImplicitModelForParameterOfOperation", "Implicit model created for parameter \"{ParameterName}\" of operation \"{OperationName}\"."),
+						TEXT("ParameterName"), FText::FromString(*InParameterName.ToString(true)),
+						TEXT("OperationName"), FText::FromString(TEXT("?")));
+					MessageLog->LogInfo(LogMessage, FWebAPIOpenAPIProvider::LogName);
+					
+					Model->Name.TypeInfo->DebugString += LogMessage.ToString();
+					
+					Model->Name.TypeInfo->JsonName = InParameterName.GetJsonName();
+					Model->Name.TypeInfo->JsonType = UWebAPIStaticTypeRegistry::ToFromJsonType;
+
+					OutParameter->Type = Model->Name;
+					OutParameter->Type.TypeInfo->Model = Model;
+				}
+
+				// Special case for "body" parameters
+				if (InParameter->In == TEXT("body"))
+				{
+					if(OutParameter->Name.ToString(true) == TEXT("Body"))
+					{
+						OutParameter->Name = InDefinitionName;
+					}
+
+					if (OutParameter->bIsArray)
+					{
+						FString Name = OutParameter->Name.ToString(true);
+						Name = ProviderSettings.Singularize(Name);
+						Name = ProviderSettings.Pluralize(Name);
+						OutParameter->Name = Name;
+					}
+				}
+
+				OutParameter->Description = InParameter->Description.Get(TEXT(""));
+				OutParameter->bIsRequired = InParameter->bRequired.Get(false);
+				OutParameter->BindToTypeInfo();
+				
+				// ConvertProperty({}, InParameterName, InParameter, InDefinitionName, OutParameter);
+
+				return true;
+			}
+
+			TObjectPtr<UWebAPIParameter> FWebAPIOpenAPISchemaConverter::ConvertParameter(const TSharedPtr<OpenAPI::V3::FParameterObject>& InSrcParameter)
+			{
+				// const FString* ParameterName = InputSchema->Components.Parameters.FindKey(InSrcParameter);
+				// check(ParameterName);
+				const FString* ParameterName = nullptr;
+
+				const FString ParameterJsonName = InSrcParameter->Name.IsEmpty() ? *ParameterName : InSrcParameter->Name;
+
+				const FWebAPITypeNameVariant ParameterTypeName = OutputSchema->TypeRegistry->GetOrMakeGeneratedType(
+					EWebAPISchemaType::Parameter,
+					ProviderSettings.MakeParameterTypeName(NameTransformer(*ParameterName)),
+					ParameterJsonName,
+					TEXT("F"));
+				
+				ParameterTypeName.TypeInfo->Suffix = TEXT("Parameter");
+
+				// Don't use ParameterTypeName->Name, it might have a Parameter specific pre/postfix.
+				FString SrcParameterDefinitionName = NameTransformer(*ParameterName);
+
+				// Will get schema or create if it doesn't exist (but will be empty)
+				const TSharedPtr<OpenAPI::V3::FSchemaObject> SrcParameterSchema = ResolveReference(InSrcParameter->Schema.Get({}), SrcParameterDefinitionName, false);
+
+				const TObjectPtr<UWebAPIParameter> DstParameter = OutputSchema->AddParameter(ParameterTypeName.TypeInfo.Get());
+
+				static TMap<FString, EWebAPIParameterStorage> InToStorage = {
+					{TEXT("query"), EWebAPIParameterStorage::Query},
+					{TEXT("header"), EWebAPIParameterStorage::Header},
+					{TEXT("path"), EWebAPIParameterStorage::Path},
+					{TEXT("formData"), EWebAPIParameterStorage::Body},
+					{TEXT("body"), EWebAPIParameterStorage::Body}
+				};
+
+				DstParameter->Storage = InToStorage[InSrcParameter->In];
+
+				DstParameter->bIsArray = IsArray(InSrcParameter);
+				if (!DstParameter->Type.HasTypeInfo())
+				{
+					// @todo:
+					/*
+					DstParameter->Type = ResolveType(InSrcParameter->Type.Get(GetDefaultJsonTypeForStorage(DstParameter->Storage)),
+						InSrcParameter->Format.Get(TEXT("")),
+						SrcParameterDefinitionName,
+						SrcParameterSchema);
+						*/
+				}
+
+				const TObjectPtr<UWebAPIModel> Model = Cast<UWebAPIModel>(DstParameter);
+				if (!ConvertModel(InSrcParameter, ParameterTypeName, Model))
+				{
+					return nullptr;
+				}
+
+				// Set the (single) property, will either be a single or array of a model or enum, based on the parameter or it's schema + parameter
+				{
+					FString PropertyName = DstParameter->bIsArray ? ProviderSettings.GetDefaultArrayPropertyName() : ProviderSettings.GetDefaultPropertyName();
+					SrcParameterDefinitionName = ProviderSettings.MakeNestedPropertyTypeName(SrcParameterDefinitionName, PropertyName);
+					
+					TObjectPtr<UWebAPIProperty>& DstProperty = DstParameter->Property = Model->Properties.Add_GetRef(NewObject<UWebAPIProperty>(Model));
+					ConvertProperty(Model->Name,
+						PropertyName,
+						InSrcParameter,
+						SrcParameterDefinitionName,
+						DstProperty);
+
+					// Add enum as it's own model, and reference it as this properties type
+					if (SrcParameterSchema.IsValid() && SrcParameterSchema->Enum.IsSet() && !SrcParameterSchema->Enum->IsEmpty())
+					{
+						const FWebAPITypeNameVariant EnumTypeName = DstParameter->Type;
+						EnumTypeName.TypeInfo->SetName(ProviderSettings.MakeParameterTypeName(*ParameterName));
+						EnumTypeName.TypeInfo->SetNested(ParameterTypeName);
+
+						const TObjectPtr<UWebAPIEnum>& Enum = ConvertEnum(SrcParameterSchema, EnumTypeName);
+
+						const FText LogMessage = FText::FormatNamed(
+							LOCTEXT("AddedImplicitEnumForParameter", "Implicit enum created for parameter \"{Name}\"."),
+							TEXT("Name"), FText::FromString(**ParameterName));
+						MessageLog->LogInfo(LogMessage, FWebAPIOpenAPIProvider::LogName);
+
+						Enum->Name.TypeInfo->DebugString += LogMessage.ToString();
+
+						DstParameter->Type = Enum->Name;
+						DstParameter->Type.TypeInfo->Model = Enum;
+					}
+					// Add struct as it's own model, and reference it as this properties type
+					else if (DstParameter->Type.ToString(true).IsEmpty())
+					{
+						const FWebAPITypeNameVariant ModelTypeName = DstParameter->Name;
+						ModelTypeName.TypeInfo->SetName(SrcParameterDefinitionName);
+						ModelTypeName.TypeInfo->Prefix = TEXT("F");
+						ModelTypeName.TypeInfo->SetNested(ParameterTypeName);
+
+						//const TObjectPtr<UWebAPIModel>& Model = OutputSchema->AddModel<UWebAPIModel>(ModelTypeName.HasTypeInfo() ? ModelTypeName.TypeInfo.Get() : nullptr);
+						ConvertModel(InSrcParameter, {}, DstParameter);
+
+						const FText LogMessage = FText::FormatNamed(
+							LOCTEXT("AddedImplicitModelForParameter", "Implicit model created for parameter \"{ParameterName}\"."),
+							TEXT("ParameterName"), FText::FromString(**ParameterName));
+						MessageLog->LogInfo(LogMessage, FWebAPIOpenAPIProvider::LogName);
+						
+						DstParameter->Name.TypeInfo->DebugString += LogMessage.ToString();
+						
+						DstParameter->Name.TypeInfo->JsonName = ParameterTypeName.GetJsonName();
+						DstParameter->Name.TypeInfo->JsonType = UWebAPIStaticTypeRegistry::ToFromJsonType;
+
+						// DstParameter->Type = Model->Name;
+						//					DstParameter->Type.TypeInfo->Model = DstParameter;
+					}
+				}
+
+	
+				/**
+				const TObjectPtr<UWebAPIModel> Model = Cast<UWebAPIModel>(DstParameter);
+				if (!ConvertModel(InSrcParameter, ParameterTypeName, Model))
+				{
+					return nullptr;
+				}
+
+				
+				if (!SrcParameterSchema.GetPath().IsEmpty())
+				{
+					SrcParameterDefinitionName = SrcParameterSchema.GetLastPathSegment();
+				}
+
+				if (SrcParameterSchema.IsSet() && SrcParameterSchema->Items.IsSet() && !SrcParameterSchema->Items->GetPath().IsEmpty())
+				{
+					SrcParameterDefinitionName = SrcParameterSchema->Items->GetLastPathSegment();
+				}
+
+				if (SrcParameterSchema.IsSet())
+				{
+					DstParameter->Type = ResolveType(SrcParameterSchema.GetShared(), NameTransformer(*ParameterName));
+				}
+				else
+				{
+					DstParameter->Type = ResolveType(InSrcParameter, NameTransformer(*ParameterName));
+				}
+				*/
+
+				// Special case for "body" parameters
+				if (InSrcParameter->In == TEXT("body"))
+				{
+					DstParameter->Name = SrcParameterDefinitionName;
+
+					if (DstParameter->bIsArray)
+					{
+						DstParameter->Name = ProviderSettings.Pluralize(SrcParameterDefinitionName);
+					}
+				}
+
+				DstParameter->BindToTypeInfo();
+
+				//DstParameter->Name = FWebAPINameInfo(NameTransformer(DstParameter->Name), DstParameter->Name.GetJsonName(), DstParameter->Type);
+				DstParameter->Description = InSrcParameter->Description.Get(TEXT(""));
+				DstParameter->bIsRequired = InSrcParameter->bRequired.Get(false);
+				
+				/*
+				if(!DstParameter->Type.TypeInfo->bIsBuiltinType)
+				{
+					if(DstParameter->Type.TypeInfo->IsEnum())
+					{
+						ConvertEnum(InSrcParameter);						
+					}
+					else
+					{
+						unimplemented();						
+					}
+				}
+				*/
+
+				/**
+				// Special case - if the name is "body", set as Body property and not Parameter
+				if (InSrcParameter->Name.Equals(TEXT("body"), ESearchCase::IgnoreCase) && !SrcParameterDefinitionName.IsEmpty())
+				{
+					// Check for existing definition
+					if (const TObjectPtr<UWebAPITypeInfo>* FoundGeneratedType = OutputSchema->TypeRegistry->FindGeneratedType(EWebAPISchemaType::Model, SrcParameterDefinitionName))
+					{
+						DstParameter->Model = Cast<UWebAPIModel>((*FoundGeneratedType)->Model.LoadSynchronous());
+						return DstParameter;
+					}
+				}
+				*/
+
+				return DstParameter;
+			}
+
+			bool FWebAPIOpenAPISchemaConverter::ConvertRequest(const FWebAPITypeNameVariant& InOperationName, const TSharedPtr<OpenAPI::V3::FOperationObject>& InOperation, const TObjectPtr<UWebAPIOperationRequest>& OutRequest)
+			{
+				// const FWebAPITypeNameVariant RequestTypeName = OutputSchema->TypeRegistry->AddUnnamedType(IWebAPIEditorModuleInterface::Get().GetStaticTypeRegistry()->Object);
+				const FWebAPITypeNameVariant RequestTypeName = OutputSchema->TypeRegistry->GetOrMakeGeneratedType(
+					EWebAPISchemaType::Model,
+					ProviderSettings.MakeRequestTypeName(InOperationName),
+					InOperationName.ToString(true),
+					IWebAPIEditorModuleInterface::Get().GetStaticTypeRegistry()->Object);
+#if WITH_WEBAPI_DEBUG
+				RequestTypeName.TypeInfo->DebugString += TEXT(">ConvertRequest");
+#endif
+				RequestTypeName.TypeInfo->Prefix = TEXT("F");
+
+				OutRequest->Name = RequestTypeName;
+
+				if(!InOperation->Parameters.IsEmpty())
+				{
+					for (Json::TJsonReference<OpenAPI::V3::FParameterObject>& SrcParameter : InOperation->Parameters)
+					{
+						FString SrcParameterDefinitionName;
+						// Will get schema or create if it doesn't exist (but will be empty)
+						TSharedPtr<OpenAPI::V3::FSchemaObject> SrcParameterSchema = ResolveReference(SrcParameter->Schema.Get({}), SrcParameterDefinitionName);
+
+						FWebAPINameVariant ParameterName = FWebAPINameInfo(NameTransformer(SrcParameter->Name), SrcParameter->Name);
+						if(SrcParameterDefinitionName.IsEmpty())
+						{
+							SrcParameterDefinitionName = ProviderSettings.MakeNestedPropertyTypeName(InOperationName, ParameterName);
+						}
+
+						const TObjectPtr<UWebAPIOperationParameter> DstParameter = OutRequest->Parameters.Add_GetRef(NewObject<UWebAPIOperationParameter>(OutRequest));
+						ConvertOperationParameter(ParameterName, SrcParameter.GetShared(), SrcParameterDefinitionName,	DstParameter);
+						DstParameter->BindToTypeInfo();
+					
+						// Special case - if the name is "body", set as Body property and not Parameter
+						if (SrcParameter->Name.Equals(TEXT("body"), ESearchCase::IgnoreCase) && !SrcParameterDefinitionName.IsEmpty())
+						{ 
+							// Check for existing definition
+							if (const TObjectPtr<UWebAPITypeInfo>* FoundGeneratedType = OutputSchema->TypeRegistry->FindGeneratedType(EWebAPISchemaType::Model, SrcParameterDefinitionName))
+							{
+								DstParameter->Model = Cast<UWebAPIModel>((*FoundGeneratedType)->Model.LoadSynchronous());
+								return true;
+							}
+						}
+					}
+				}
+				else if(TSharedPtr<OpenAPI::V3::FRequestBodyObject> SrcRequestBody = ResolveReference(InOperation->RequestBody))
+				{
+					if(!SrcRequestBody->Content.IsEmpty())
+					{
+						TSharedPtr<V3::FMediaTypeObject> MediaType = nullptr;
+					
+						// Choose json entry or whatever is first (if json not found)
+						if(const TSharedPtr<V3::FMediaTypeObject>* JsonMediaType = SrcRequestBody->Content.Find(UE::WebAPI::MimeType::NAME_Json.ToString()))
+						{
+							MediaType = *JsonMediaType;
+						}
+						else if(!SrcRequestBody->Content.IsEmpty())
+						{
+							TArray<FString> Keys;
+							SrcRequestBody->Content.GetKeys(Keys);
+							MediaType = SrcRequestBody->Content[Keys[0]];
+						}
+
+						if(MediaType.IsValid())
+						{
+							FString MediaTypeDefinitionName;
+							TSharedPtr<OpenAPI::V3::FSchemaObject> MediaTypeSchema = ResolveReference(MediaType->Schema, MediaTypeDefinitionName);
+
+							// @todo: what if it's an array of $ref? TypeInfo doesn't specify whether it's an array or not - check spec
+							const TObjectPtr<UWebAPIOperationParameter> DstParameter = OutRequest->Parameters.Add_GetRef(NewObject<UWebAPIOperationParameter>(OutRequest));
+							DstParameter->Type = ResolveType(MediaTypeSchema, MediaTypeDefinitionName);
+							check(DstParameter->Type.IsValid());
+							
+							DstParameter->Name = ProviderSettings.GetDefaultPropertyName();
+							DstParameter->Name = ResolvePropertyName(DstParameter, DstParameter->Type.ToString(true));
+							check(DstParameter->Name.IsValid());
+
+							DstParameter->Storage = EWebAPIParameterStorage::Body;
+							DstParameter->BindToTypeInfo();
+						}
+					}
+				}
+				// Operation has no parameters, but if arbitrary json is enabled, make a Value property of type JsonObject
+				else if(ProviderSettings.bEnableArbitraryJsonPayloads)
+				{
+					FWebAPINameVariant ParameterName = TEXT("Value");
+
+					const TObjectPtr<UWebAPIOperationParameter> DstParameter = OutRequest->Parameters.Add_GetRef(NewObject<UWebAPIOperationParameter>(OutRequest));
+					DstParameter->Type = IWebAPIEditorModuleInterface::Get().GetStaticTypeRegistry()->JsonObject;
+					DstParameter->Name = FWebAPINameInfo(ParameterName.ToString(true), ParameterName.GetJsonName(), DstParameter->Type);
+					DstParameter->bIsMixin = true; // treat as object, not as field
+					DstParameter->Storage = EWebAPIParameterStorage::Body;
+					DstParameter->BindToTypeInfo();
+				}
+
+				OutRequest->BindToTypeInfo();
+				
+				return true;
+			}
+
+			bool FWebAPIOpenAPISchemaConverter::ConvertResponse(const FWebAPITypeNameVariant& InOperationName, uint32 InResponseCode, const TSharedPtr<OpenAPI::V3::FResponseObject>& InResponse, const TObjectPtr<UWebAPIOperationResponse>& OutResponse)
+			{
+				const FWebAPITypeNameVariant ResponseTypeName = OutputSchema->TypeRegistry->GetOrMakeGeneratedType(
+					EWebAPISchemaType::Model,
+					ProviderSettings.MakeResponseTypeName(InOperationName, InResponseCode),
+					InOperationName.ToString(true),
+					IWebAPIEditorModuleInterface::Get().GetStaticTypeRegistry()->Object);
+#if WITH_WEBAPI_DEBUG
+				ResponseTypeName.TypeInfo->DebugString += TEXT(">ConvertResponse");
+#endif
+				ResponseTypeName.TypeInfo->Prefix = TEXT("F");
+
+				check(ResponseTypeName.HasTypeInfo());
+
+				OutResponse->Name = ResponseTypeName;
+				OutResponse->Code = InResponseCode;
+				OutResponse->Description = InResponse->Description;
+				OutResponse->Message = OutResponse->Description;
+
+				if (InResponse->Content.IsSet())
+				{
+					for(const TPair<FString, Json::TJsonReference<V3::FMediaTypeObject>>& Item : InResponse->Content.GetValue())
+					{
+						// Will get schema or create if it doesn't exist (but will be empty)
+						FString SrcPropertyDefinitionName;
+						TSharedPtr<V3::FSchemaObject> SrcResponseSchema = ResolveReference(Item.Value->Schema, SrcPropertyDefinitionName);
+
+						const FString PropertyName = SrcResponseSchema->Type.Get(GetDefaultJsonTypeForStorage(OutResponse->Storage)) == TEXT("array") ? ProviderSettings.GetDefaultArrayPropertyName() : ProviderSettings.GetDefaultPropertyName();
+						const TObjectPtr<UWebAPIProperty>& DstProperty = OutResponse->Properties.Add_GetRef(NewObject<UWebAPIProperty>(OutResponse));
+						DstProperty->bIsMixin = true;
+						ConvertProperty(OutResponse->Name,
+							PropertyName,
+							SrcResponseSchema,
+							SrcPropertyDefinitionName,
+							DstProperty);
+
+						/*
+						//OutResponse->bIsArray = IsArray(InSrcParameter);
+						if (!DstProperty->Type.HasTypeInfo())
+						{
+							DstProperty->Type = ResolveType(SrcResponseSchema->Type.Get(TEXT("object")),
+								SrcResponseSchema->Format.Get(TEXT("")),
+								SrcPropertyDefinitionName,
+								SrcResponseSchema.GetShared());
+						}
+
+						
+						if(SrcPropertyDefinitionName.IsEmpty())
+						{
+							SrcPropertyDefinitionName = ProviderSettings.MakeNestedPropertyTypeName(InOperationName, PropertyName);
+						}
+
+						const FString Type = InResponse->Schema->Get()->Type.Get(TEXT("object"));
+
+						const TObjectPtr<UWebAPIModel> Model = OutResponse;
+						if (!ConvertModel(SrcResponseSchema.GetShared(), ResponseTypeName, Model))
+						{
+							return false;
+						}
+
+						// Don't use ParameterTypeName->Name, it might have a Parameter specific pre/postfix.
+						FString SrcResponseDefinitionName = ResponseTypeName.ToString(true);
+
+						// Will get schema or create if it doesn't exist (but will be empty)
+						const TSharedPtr<OpenAPI::V3::FSchemaObject> SrcParameterSchema = ResolveReference(InResponse->Schema.Get({}), SrcResponseDefinitionName, false);
+
+						const FString PropertyName = Type == TEXT("array") ? ProviderSettings.GetDefaultArrayPropertyName() : ProviderSettings.GetDefaultPropertyName();
+						// FString DefinitionName = ProviderSettings.MakeNestedPropertyTypeName(DefinitionName, PropertyName);
+						
+						const TObjectPtr<UWebAPIProperty>& DstProperty = OutResponse->Properties.Add_GetRef(NewObject<UWebAPIProperty>(OutResponse));
+						ConvertProperty(OutResponse->Name,
+							PropertyName,
+							SrcParameterSchema,
+							SrcResponseDefinitionName,
+							DstProperty);
+						
+
+						if (Type == TEXT("array"))
+						{
+	 
+						}
+						*/
+						// @todo: Account for single item = single property, ie. type=array, items=SomeDef would be TArray<FSomeDef> Items 
+					}
+				}
+
+				if (InResponse->Headers.IsSet())
+				{
+					for (const TPair<FString, Json::TJsonReference<OpenAPI::V3::FHeaderObject>>& SrcHeader : InResponse->Headers.GetValue())
+					{
+						FString Key = SrcHeader.Key; // @todo: what is this?
+						// SrcHeader.Value-> // @todo: this describes return object
+					}
+				}
+
+				OutResponse->BindToTypeInfo();
+
+				return true;
+			}
+
+			TObjectPtr<UWebAPIOperation> FWebAPIOpenAPISchemaConverter::ConvertOperation(const FString& InPath, const FString& InVerb, const TSharedPtr<OpenAPI::V3::FOperationObject>& InSrcOperation, const FWebAPITypeNameVariant& InOperationTypeName)
+			{
+				
+				FString OperationName = InSrcOperation->OperationId.Get(InSrcOperation->Summary.Get(TEXT("")));
+				if(OperationName.IsEmpty())
+				{
+					// ie. GET /pets/{id} == GetPetsById
+					
+					TArray<FString> SplitPath;
+					InPath.ParseIntoArray(SplitPath, TEXT("/"));
+
+					Algo::ForEach(SplitPath, [&](FString& Str)
+					{
+						Str = NameTransformer(Str);
+					});
+
+					OperationName = NameTransformer(InVerb) + FString::Join(SplitPath, TEXT(""));
+					
+				}
+				
+				OperationName = NameTransformer(OperationName);
+				OperationName = ProviderSettings.MakeValidMemberName(OperationName, InVerb);
+
+				check(!OperationName.IsEmpty());
+				
+				FWebAPITypeNameVariant OperationTypeName;
+				if(InOperationTypeName.IsValid())
+				{
+					OperationName = InOperationTypeName.ToString(true);
+					OperationTypeName = InOperationTypeName;
+				}
+				else
+				{
+					OperationTypeName = OutputSchema->TypeRegistry->GetOrMakeGeneratedType(
+						EWebAPISchemaType::Operation,
+						NameTransformer(*OperationName),
+						*OperationName,
+						TEXT("U"));
+					OperationTypeName.TypeInfo->SetName(OperationName);
+
+#if WITH_WEBAPI_DEBUG
+					OperationTypeName.TypeInfo->DebugString += TEXT(">ConvertOperation");
+#endif
+				}
+
+				// A spec can have no tags, so ensure there's a default
+				const FString FirstTag = ProviderSettings.ToPascalCase(InSrcOperation->Tags.Get({TEXT("Default")})[0]);
+				TObjectPtr<UWebAPIService>* Service = OutputSchema->Services.Find(FirstTag);
+				checkf(Service, TEXT("An operation must belong to a service!"));
+
+				const TObjectPtr<UWebAPIOperation> DstOperation = (*Service)->Operations.Add_GetRef(NewObject<UWebAPIOperation>(*Service));
+				DstOperation->Service = *Service;
+				DstOperation->Verb = InVerb;
+				DstOperation->Path = InPath;
+
+				DstOperation->Name = OperationTypeName;
+				// Choose first non-empty of: description, summary, operation id
+				DstOperation->Description = InSrcOperation->Description.Get(
+					InSrcOperation->OperationId.Get(InSrcOperation->Summary.Get(
+						InSrcOperation->OperationId.Get(TEXT("")))));
+				DstOperation->bIsDeprecated = InSrcOperation->bDeprecated.Get(false);
+
+				ConvertRequest(DstOperation->Name, InSrcOperation, DstOperation->Request);
+
+				// @todo:
+				/*
+				if (InSrcOperation->Consumes.IsSet())
+				{
+					DstOperation->RequestContentTypes.Append(InSrcOperation->Consumes.GetValue());
+				}
+
+				if (InSrcOperation->Produces.IsSet())
+				{
+					DstOperation->ResponseContentTypes.Append(InSrcOperation->Produces.GetValue());
+				}
+				*/
+
+				if (ensure(!InSrcOperation->Responses.IsEmpty()))
+				{
+					for (const TPair<FString, Json::TJsonReference<OpenAPI::V3::FResponseObject>>& SrcResponse : InSrcOperation->Responses)
+					{
+						// If "Default", resolves to 0, and means all other unhandled codes (similar to default in switch statement)
+						const uint32 Code = FCString::Atoi(*SrcResponse.Key);
+						const TObjectPtr<UWebAPIOperationResponse> DstResponse = DstOperation->Responses.Add_GetRef(NewObject<UWebAPIOperationResponse>(DstOperation));
+						// OutOperation->AddResponse() @todo
+						ConvertResponse(DstOperation->Name, Code, ResolveReference(SrcResponse.Value), DstResponse);
+
+						// If success response (code 200), had no resolved properties but the operation says it returns something, then add that something as a property
+						if(DstResponse->Properties.IsEmpty() && !DstOperation->ResponseContentTypes.IsEmpty() && DstResponse->Code == 200)
+						{
+							const FString PropertyName = ProviderSettings.GetDefaultPropertyName();
+							const TObjectPtr<UWebAPIProperty>& DstProperty = DstResponse->Properties.Add_GetRef(NewObject<UWebAPIProperty>(DstResponse));
+							DstProperty->bIsMixin = true;
+							ConvertProperty(DstResponse->Name,
+								PropertyName,
+								GetTypeForContentType(DstOperation->ResponseContentTypes[0]),
+								DstProperty);
+							DstProperty->Name = PropertyName;
+							DstResponse->Storage = EWebAPIResponseStorage::Body;
+						}
+					}
+				}
+
+				check(DstOperation->Name.HasTypeInfo());
+
+				const FName OperationObjectName = ProviderSettings.MakeOperationObjectName(*Service, OperationName);
+				DstOperation->Rename(*OperationObjectName.ToString(), DstOperation->GetOuter());
+				
+				return DstOperation;
+			}
+
+			TObjectPtr<UWebAPIService> FWebAPIOpenAPISchemaConverter::ConvertService(const FWebAPINameVariant& InName) const
+			{
+				const FString TagName = NameTransformer(InName);
+				return OutputSchema->GetOrMakeService(TagName);
+			}
+
+			TObjectPtr<UWebAPIService> FWebAPIOpenAPISchemaConverter::ConvertService(const TSharedPtr<OpenAPI::V3::FTagObject>& InTag) const
+			{
+				const TObjectPtr<UWebAPIService> Service = ConvertService(InTag->Name);
+				Service->Description = InTag->Description.Get(TEXT(""));
+				return Service;
+			}
+
+			bool FWebAPIOpenAPISchemaConverter::ConvertModels(const TMap<FString, Json::TJsonReference<OpenAPI::V3::FSchemaObject>>& InSchemas,	UWebAPISchema* OutSchema)
+			{
+				bool bAllConverted = true;
+				for (const TTuple<FString, Json::TJsonReference<OpenAPI::V3::FSchemaObject>>& Item : InSchemas)
+				{
+					const FString& Name = Item.Key;
+					const TSharedPtr<V3::FSchemaObject> Schema = ResolveReference(Item.Value);
+
+					if (!Schema.IsValid())
+					{
+						bAllConverted = false;
+						FFormatNamedArguments Args;
+						Args.Add(TEXT("ModelName"), FText::FromString(Name));
+						MessageLog->LogWarning(
+							FText::Format(LOCTEXT("SchemaInvalid", "The schema for model \"{ModelName}\" was invalid/null."), Args),
+							FWebAPIOpenAPIProvider::LogName);
+						continue;
+					}
+
+					FWebAPITypeNameVariant ModelName = OutputSchema->TypeRegistry->GetOrMakeGeneratedType(
+						EWebAPISchemaType::Model,
+						NameTransformer(Name),
+						Name,
+						TEXT("F"));
+
+					const TObjectPtr<UWebAPIModel>& Model = OutSchema->AddModel<UWebAPIModel>(ModelName.TypeInfo.Get());
+					ConvertModel(Schema, {}, Model);
+				}
+
+				// Set Model property of TypeInfos where applicable
+				for (const TObjectPtr<UWebAPIModelBase>& ModelBase : OutputSchema->Models)
+				{
+					if (const TObjectPtr<UWebAPIModel>& Model = Cast<UWebAPIModel>(ModelBase))
+					{
+						Model->BindToTypeInfo();
+						// Model->Name.TypeInfo->Model = Model;
+					}
+					else if (const TObjectPtr<UWebAPIEnum>& Enum = Cast<UWebAPIEnum>(ModelBase))
+					{
+						Enum->BindToTypeInfo();
+						// Enum->Name.TypeInfo->Model = Enum;
+					}
+					else if (const TObjectPtr<UWebAPIParameter>& ParameterModel = Cast<UWebAPIParameter>(ModelBase))
+					{
+						ParameterModel->BindToTypeInfo();
+						// ParameterModel->Type.TypeInfo->Model = ParameterModel;
+					}
+				}
+
+				return bAllConverted;
+			}
+
+			bool FWebAPIOpenAPISchemaConverter::ConvertParameters(const TArray<Json::TJsonReference<OpenAPI::V3::FParameterObject>>& InParameters, UWebAPISchema* OutSchema)
+			{
+				bool bAllConverted = true;
+				for (const Json::TJsonReference<OpenAPI::V3::FParameterObject>& Item : InParameters)
+				{
+					bAllConverted &= ConvertParameter(ResolveReference(Item)) != nullptr;
+				}
+
+				// Re-bind models to their TypeInfos
+				for (const TObjectPtr<UWebAPIModelBase>& Model : OutputSchema->Models)
+				{
+					Model->BindToTypeInfo();
+				}
+
+				return bAllConverted;
+			}
+
+			bool FWebAPIOpenAPISchemaConverter::ConvertParameters(const TMap<FString, TSharedPtr<OpenAPI::V3::FParameterObject>>& InParameters, UWebAPISchema* OutSchema)
+			{
+				bool bAllConverted = true;
+				for (const TTuple<FString, TSharedPtr<OpenAPI::V3::FParameterObject>>& Item : InParameters)
+				{
+					const FWebAPITypeNameVariant& Name = Item.Key;
+					const TSharedPtr<OpenAPI::V3::FParameterObject>& SrcParameter = Item.Value;
+
+					bAllConverted &= ConvertParameter(SrcParameter) != nullptr;
+				}
+
+				// Re-bind models to their TypeInfos
+				for (const TObjectPtr<UWebAPIModelBase>& Model : OutputSchema->Models)
+				{
+					Model->BindToTypeInfo();
+				}
+
+				return bAllConverted;
+			}
+
+			bool FWebAPIOpenAPISchemaConverter::ConvertSecurity(const UE::WebAPI::OpenAPI::V3::FOpenAPIObject& InOpenAPI, UWebAPISchema* OutSchema)
+			{
+				if (!InOpenAPI.Components->SecuritySchemes.IsEmpty())
+				{
+					for (const TPair<FString, Json::TJsonReference<OpenAPI::V3::FSecuritySchemeObject>>& Item : InOpenAPI.Components->SecuritySchemes)
+					{
+						const FString& Name = Item.Key;
+						const TSharedPtr<OpenAPI::V3::FSecuritySchemeObject>& SecurityScheme = ResolveReference(Item.Value);
+						if(SecurityScheme->Type == TEXT("http") && SecurityScheme->Scheme == TEXT("basic"))
+						{
+							// type = basic
+						}
+						else if(SecurityScheme->Type == TEXT("http") && SecurityScheme->Scheme == TEXT("bearer"))
+						{
+							// type = apiKey
+							// name = authorization
+							// storage = header
+						}
+						else if(SecurityScheme->Type == TEXT("oauth2"))
+						{
+							// @todo:
+							FString FlowName = SecurityScheme->Flows.IsValid() ? TEXT("") : TEXT("");
+							V3::FOAuthFlowsObject Flow = (SecurityScheme->Flows.Get())[0];
+
+							if(FlowName == TEXT("clientCredentials"))
+							{
+								// flow = application
+							}
+							else if(FlowName == TEXT("authorizationCode"))
+							{
+								// flow = accessCode
+							}
+							else
+							{
+								// flow = FlowName
+							}
+
+							
+						}
+
+						// @todo: WebAPI security primitive
+						// TObjectPtr<UWebAPIModel> Model = OutSchema->Client->Models.Add_GetRef(NewObject<UWebAPIModel>(OutSchema));
+						// ConvertModel(Name, Schema.Get(), Model);
+					}
+
+					return true;
+				}
+
+				return false;
+			}
+
+			bool FWebAPIOpenAPISchemaConverter::ConvertTags(const TArray<TSharedPtr<OpenAPI::V3::FTagObject>>& InTags, UWebAPISchema* OutSchema) const
+			{
+				for (const TSharedPtr<OpenAPI::V3::FTagObject>& Tag : InTags)
+				{
+					ConvertService(Tag);
+				}
+
+				return OutSchema->Services.Num() > 0;
+			}
+
+			bool FWebAPIOpenAPISchemaConverter::ConvertPaths(const TMap<FString, TSharedPtr<OpenAPI::V3::FPathItemObject>>& InPaths, UWebAPISchema* OutSchema)
+			{
+				for (const TTuple<FString, TSharedPtr<OpenAPI::V3::FPathItemObject>>& Item : InPaths)
+				{
+					const FString& Url = Item.Key;
+					const TSharedPtr<OpenAPI::V3::FPathItemObject>& Path = Item.Value;
+
+					ConvertParameters(Path->Parameters, OutputSchema.Get());
+					
+					TMap<FString, TSharedPtr<OpenAPI::V3::FOperationObject>> SrcVerbs;
+					if (Path->Get.IsValid())
+					{
+						SrcVerbs.Add(TEXT("Get"), Path->Get);
+					}
+					if (Path->Put.IsValid())
+					{
+						SrcVerbs.Add(TEXT("Put"), Path->Put);
+					}
+					if (Path->Post.IsValid())
+					{
+						SrcVerbs.Add(TEXT("Post"), Path->Post);
+					}
+					if (Path->Delete.IsValid())
+					{
+						SrcVerbs.Add(TEXT("Delete"), Path->Delete);
+					}
+					if (Path->Options.IsValid())
+					{
+						SrcVerbs.Add(TEXT("Options"), Path->Options);
+					}
+					if (Path->Head.IsValid())
+					{
+						SrcVerbs.Add(TEXT("Head"), Path->Head);
+					}
+					if (Path->Patch.IsValid())
+					{
+						SrcVerbs.Add(TEXT("Patch"), Path->Patch);
+					}
+
+					// Each path can have multiple, ie. Get, Put, Delete
+					for (const TPair<FString, TSharedPtr<OpenAPI::V3::FOperationObject>>& VerbOperationPair : SrcVerbs)
+					{
+						FString Verb = VerbOperationPair.Key;
+						if (VerbOperationPair.Value == nullptr || !VerbOperationPair.Value.IsValid())
+						{
+							continue;
+						}
+
+						TSharedPtr<OpenAPI::V3::FOperationObject> SrcOperation = VerbOperationPair.Value;
+						if(SrcOperation->RequestBody.IsSet())
+						{
+							FString ParameterDefinitionName;
+							// @todo:
+							/*
+							TSharedPtr<V3::FRequestBodyObject> Parameter = ResolveReference(SrcOperation->RequestBody.GetShared(), ParameterDefinitionName);
+							if(!Parameter->Content.IsEmpty())
+							{
+								auto Content = Parameter->Content;
+									
+							}
+							*/
+						}
+
+						TArray<FString> Tags;
+						if (SrcOperation.IsValid() && SrcOperation->Tags.IsSet())
+						{
+							Tags = SrcOperation->Tags.GetValue();
+						}
+						else
+						{
+							Tags = { TEXT("Default") };
+						}
+
+						for (FString& Tag : Tags)
+						{
+							// Do first to retain original case
+							TObjectPtr<UWebAPIService> Service = ConvertService(Tag);
+							Tag = NameTransformer(Tag);
+
+							FString OperationName = SrcOperation->OperationId.Get(Verb + Tag);
+							OperationName = NameTransformer(OperationName);
+							OperationName = ProviderSettings.MakeValidMemberName(OperationName, Verb);
+
+							FString OperationNamePrefix = TEXT("");
+							int32 OperationNameSuffix = 0;
+
+							// Check if this already exists - each operation is unique so the name needs to be different - prepend verb
+							while(OutSchema->TypeRegistry->FindGeneratedType(EWebAPISchemaType::Operation,
+																			 OperationNamePrefix + OperationName + (OperationNameSuffix > 0 ? FString::FormatAsNumber(OperationNameSuffix) : TEXT(""))) != nullptr)
+							{
+								if(!OperationName.StartsWith(Verb))
+								{
+									OperationNamePrefix = Verb;
+								}
+								OperationNameSuffix += 1; 
+							}
+
+							OperationName = OperationNamePrefix + OperationName + (OperationNameSuffix > 0 ? FString::FormatAsNumber(OperationNameSuffix) : TEXT(""));
+				
+							TObjectPtr<UWebAPITypeInfo> OperationTypeInfo = OutSchema->TypeRegistry->
+								GetOrMakeGeneratedType(
+													   EWebAPISchemaType::Operation,
+													   OperationName,
+													   OperationName,
+													   TEXT("U"));
+
+							const TObjectPtr<UWebAPIOperation> Operation = ConvertOperation(Url, Verb, SrcOperation, OperationTypeInfo);
+							Operation->Service = Service;
+							Operation->Verb = Verb;
+							Operation->Path = Url;
+
+							FStringFormatNamedArguments FormatArgs;
+							FormatArgs.Add(TEXT("ClassName"), UWebAPIOperation::StaticClass()->GetName());
+							FormatArgs.Add(TEXT("ServiceName"), Service->Name.ToString(true));
+							FormatArgs.Add(TEXT("OperationName"), OperationName);
+
+							const FName OperationObjectName = ProviderSettings.MakeOperationObjectName(Service, OperationName);
+							Operation->Rename(*OperationObjectName.ToString(), Operation->GetOuter());
+
+							Operation->BindToTypeInfo();
+
+							// Verb->Parameters
+							// Verb->Responses
+							// Verb->Schemes
+							// Verb->Consumes ie. "application/json"
+							// Verb->Produces ie. "application/json"
+							// Verb->Security							
+						}
+						// @todo:
+						// ConvertParameters(Operation)
+
+						// @todo:
+						//ConvertResponses()
+					}
+				}
+				return true;
+			}
+		};
+	}
+};
+
+#undef SET_OPTIONAL_FLAGGED
+#undef SET_OPTIONAL
+
+#undef LOCTEXT_NAMESPACE
+
+// @todo: remove
+PRAGMA_ENABLE_OPTIMIZATION
diff --git a/Engine/Plugins/Web/WebAPI/Source/WebAPIOpenAPI/Private/V3/WebApiOpenApiSchema.cpp b/Engine/Plugins/Web/WebAPI/Source/WebAPIOpenAPI/Private/V3/WebApiOpenApiSchema.cpp
new file mode 100644
--- /dev/null
+++ b/Engine/Plugins/Web/WebAPI/Source/WebAPIOpenAPI/Private/V3/WebApiOpenApiSchema.cpp
@@ -0,0 +1,333 @@
+﻿// Copyright Epic Games, Inc. All Rights Reserved.
+
+#include "WebAPIJsonUtilities.h"
+
+#include "Dom/JsonObject.h"
+#include "V3/WebAPIOpenAPISchema.h"
+
+namespace UE::WebAPI::OpenAPI::V3
+{
+	bool FSchemaObjectBase::FromJson(const TSharedRef<FJsonObject>& InJson)
+	{
+		bool bSetAField = true;
+		bSetAField &= Json::TryGetField(InJson, TEXT("name"), Name);
+		bSetAField &= Json::TryGetField(InJson, TEXT("description"), Description);
+		bSetAField &= Json::TryGetField(InJson, TEXT("required"), bRequired);
+		bSetAField &= Json::TryGetField(InJson, TEXT("deprecated"), bDeprecated);
+		bSetAField &= Json::TryGetField(InJson, TEXT("example"), Example);
+		return bSetAField;
+	}
+
+	bool FParameterObject::FromJson(const TSharedRef<FJsonObject>& InJson)
+	{
+		bool bSetAField = Super::FromJson(InJson);
+		bSetAField &= Json::TryGetField(InJson, TEXT("in"), In);
+		bSetAField &= Json::TryGetField(InJson, TEXT("allowEmptyValue"), bAllowEmptyValue);
+		bSetAField &= Json::TryGetField(InJson, TEXT("style"), Style);
+		bSetAField &= Json::TryGetField(InJson, TEXT("explode"), bExplode);
+		bSetAField &= Json::TryGetField(InJson, TEXT("allowReserved"), bAllowReserved);
+		bSetAField &= Json::TryGetField(InJson, TEXT("schema"), Schema);
+		bSetAField &= Json::TryGetField(InJson, TEXT("examples"), Examples);
+		bSetAField &= Json::TryGetField(InJson, TEXT("content"), Content);
+		return bSetAField;
+	}
+	
+	bool FComponentsObject::FromJson(const TSharedRef<FJsonObject>& InJson)
+	{
+		bool bSetAField = true;
+		bSetAField &= Json::TryGetField(InJson, TEXT("schemas"), Schemas);
+		bSetAField &= Json::TryGetField(InJson, TEXT("responses"), Responses);
+		bSetAField &= Json::TryGetField(InJson, TEXT("parameters"), Parameters);
+		bSetAField &= Json::TryGetField(InJson, TEXT("examples"), Examples);
+		bSetAField &= Json::TryGetField(InJson, TEXT("requestBodies"), RequestBodies);
+		bSetAField &= Json::TryGetField(InJson, TEXT("headers"), Headers);
+		bSetAField &= Json::TryGetField(InJson, TEXT("securitySchemes"), SecuritySchemes);
+		bSetAField &= Json::TryGetField(InJson, TEXT("links"), Links);
+		bSetAField &= Json::TryGetField(InJson, TEXT("callbacks"), Callbacks);
+		return bSetAField;
+	}
+
+	bool FContactObject::FromJson(const TSharedRef<FJsonObject>& InJson)
+	{
+		bool bSetAField = true;
+		bSetAField &= Json::TryGetField(InJson, TEXT("name"), Name);
+		bSetAField &= Json::TryGetField(InJson, TEXT("url"), Url);
+		bSetAField &= Json::TryGetField(InJson, TEXT("email"), Email);
+		return bSetAField;
+	}
+
+	bool FDiscriminatorObject::FromJson(const TSharedRef<FJsonObject>& InJson)
+	{
+		bool bSetAField = true;
+		bSetAField &= Json::TryGetField(InJson, TEXT("propertyName"), PropertyName);
+		bSetAField &= Json::TryGetField(InJson, TEXT("mapping"), Mapping);
+		return bSetAField;
+	}
+
+	bool FEncodingObject::FromJson(const TSharedRef<FJsonObject>& InJson)
+	{
+		bool bSetAField = true;
+		bSetAField &= Json::TryGetField(InJson, TEXT("contentType"), ContentType);
+		bSetAField &= Json::TryGetField(InJson, TEXT("headers"), Headers);
+		bSetAField &= Json::TryGetField(InJson, TEXT("style"), Style);
+		bSetAField &= Json::TryGetField(InJson, TEXT("explode"), bExplode);
+		bSetAField &= Json::TryGetField(InJson, TEXT("allowReserved"), bAllowReserved);
+		return bSetAField;
+	}
+
+	bool FExampleObject::FromJson(const TSharedRef<FJsonObject>& InJson)
+	{
+		bool bSetAField = true;
+		bSetAField &= Json::TryGetField(InJson, TEXT("summary"), Summary);
+		bSetAField &= Json::TryGetField(InJson, TEXT("description"), Description);
+		// bSetAField &= Json::TryGetField(InJson, TEXT("value"), Value);
+		bSetAField &= Json::TryGetField(InJson, TEXT("externalValue"), ExternalValue);
+		return bSetAField;
+	}
+
+	bool FExternalDocumentationObject::FromJson(const TSharedRef<FJsonObject>& InJson)
+	{
+		bool bSetAField = true;
+		bSetAField &= Json::TryGetField(InJson, TEXT("description"), Description);
+		bSetAField &= Json::TryGetField(InJson, TEXT("url"), Url);
+		return bSetAField;
+	}
+
+	bool FHeaderObject::FromJson(const TSharedRef<FJsonObject>& InJson)
+	{
+		const bool bSetAField = true;
+		// @todo
+		return bSetAField;
+	}
+
+	bool FInfoObject::FromJson(const TSharedRef<FJsonObject>& InJson)
+	{
+		bool bSetAField = true;
+		bSetAField &= Json::TryGetField(InJson, TEXT("title"), Title);
+		bSetAField &= Json::TryGetField(InJson, TEXT("description"), Description);
+		bSetAField &= Json::TryGetField(InJson, TEXT("termsOfService"), TermsOfService);
+		bSetAField &= Json::TryGetField(InJson, TEXT("contact"), Contact);
+		bSetAField &= Json::TryGetField(InJson, TEXT("license"), License);
+		bSetAField &= Json::TryGetField(InJson, TEXT("version"), Version);
+		return bSetAField;
+	}
+
+	bool FLicenseObject::FromJson(const TSharedRef<FJsonObject>& InJson)
+	{
+		bool bSetAField = true;
+		bSetAField &= Json::TryGetField(InJson, TEXT("name"), Name);
+		bSetAField &= Json::TryGetField(InJson, TEXT("url"), Url);
+		return bSetAField;
+	}
+
+	bool FLinkObject::FromJson(const TSharedRef<FJsonObject>& InJson)
+	{
+		bool bSetAField = true;
+		bSetAField &= Json::TryGetField(InJson, TEXT("operationRef"), OperationRef);
+		bSetAField &= Json::TryGetField(InJson, TEXT("operationId"), OperationId);
+		bSetAField &= Json::TryGetField(InJson, TEXT("description"), Description);
+		bSetAField &= Json::TryGetField(InJson, TEXT("server"), Server);
+		return bSetAField;
+	}
+
+	bool FMediaTypeObject::FromJson(const TSharedRef<FJsonObject>& InJson)
+	{
+		bool bSetAField = true;
+		bSetAField &= Json::TryGetField(InJson, TEXT("schema"), Schema);
+		bSetAField &= Json::TryGetField(InJson, TEXT("example"), Example);
+		bSetAField &= Json::TryGetField(InJson, TEXT("examples"), Examples);
+		bSetAField &= Json::TryGetField(InJson, TEXT("encoding"), Encoding);
+		return bSetAField;
+	}
+
+	bool FOAuthFlowObject::FromJson(const TSharedRef<FJsonObject>& InJson)
+	{
+		bool bSetAField = true;
+		bSetAField &= Json::TryGetField(InJson, TEXT("authorizationUrl"), AuthorizationUrl);
+		bSetAField &= Json::TryGetField(InJson, TEXT("tokenUrl"), TokenUrl);
+		bSetAField &= Json::TryGetField(InJson, TEXT("refreshUrl"), RefreshUrl);
+		bSetAField &= Json::TryGetField(InJson, TEXT("scopes"), Scopes);
+		return bSetAField;
+	}
+
+	bool FOAuthFlowsObject::FromJson(const TSharedRef<FJsonObject>& InJson)
+	{
+		bool bSetAField = true;
+		bSetAField &= Json::TryGetField(InJson, TEXT("implicit"), Implicit);
+		bSetAField &= Json::TryGetField(InJson, TEXT("password"), Password);
+		bSetAField &= Json::TryGetField(InJson, TEXT("clientCredentials"), ClientCredentials);
+		bSetAField &= Json::TryGetField(InJson, TEXT("authorizationCode"), AuthorizationCode);
+		return bSetAField;
+	}
+
+	bool FOpenAPIObject::FromJson(const TSharedRef<FJsonObject>& InJson)
+	{
+		bool bSetAField = true;
+		bSetAField &= Json::TryGetField(InJson, TEXT("openapi"), OpenAPI);
+		bSetAField &= Json::TryGetField(InJson, TEXT("info"), Info);
+		bSetAField &= Json::TryGetField(InJson, TEXT("servers"), Servers);
+		bSetAField &= Json::TryGetField(InJson, TEXT("paths"), Paths);
+		bSetAField &= Json::TryGetField(InJson, TEXT("components"), Components);
+		bSetAField &= Json::TryGetField(InJson, TEXT("security"), Security);
+		bSetAField &= Json::TryGetField(InJson, TEXT("tags"), Tags);
+		bSetAField &= Json::TryGetField(InJson, TEXT("externalDocs"), ExternalDocs);
+		return bSetAField;
+	}
+
+	bool FOperationObject::FromJson(const TSharedRef<FJsonObject>& InJson)
+	{
+		bool bSetAField = true;
+		bSetAField &= Json::TryGetField(InJson, TEXT("tags"), Tags);
+		bSetAField &= Json::TryGetField<FString>(InJson, TEXT("summary"), Summary);
+		bSetAField &= Json::TryGetField(InJson, TEXT("description"), Description);
+		bSetAField &= Json::TryGetField(InJson, TEXT("externalDocs"), ExternalDocs);
+		bSetAField &= Json::TryGetField(InJson, TEXT("operationId"), OperationId);
+		bSetAField &= Json::TryGetField(InJson, TEXT("parameters"), Parameters);
+		bSetAField &= Json::TryGetField(InJson, TEXT("requestBody"), RequestBody);
+		bSetAField &= Json::TryGetField(InJson, TEXT("responses"), Responses);
+		bSetAField &= Json::TryGetField(InJson, TEXT("callbacks"), Callbacks);
+		bSetAField &= Json::TryGetField(InJson, TEXT("deprecated"), bDeprecated);
+		bSetAField &= Json::TryGetField(InJson, TEXT("security"), Security);
+		bSetAField &= Json::TryGetField(InJson, TEXT("servers"), Servers);
+		return bSetAField;
+	}
+
+	bool FReferenceObject::FromJson(const TSharedRef<FJsonObject>& InJson)
+	{
+		bool bSetAField = true;
+		bSetAField &= Json::TryGetField(InJson, TEXT("$ref"), $Ref);
+		return bSetAField;
+	}
+
+	bool FRequestBodyObject::FromJson(const TSharedRef<FJsonObject>& InJson)
+	{
+		bool bSetAField = true;
+		bSetAField &= Json::TryGetField(InJson, TEXT("description"), Description);
+		bSetAField &= Json::TryGetField(InJson, TEXT("content"), Content);
+		bSetAField &= Json::TryGetField(InJson, TEXT("required"), bRequired);
+		return bSetAField;
+	}
+
+	bool FResponseObject::FromJson(const TSharedRef<FJsonObject>& InJson)
+	{
+		bool bSetAField = true;
+		bSetAField &= Json::TryGetField(InJson, TEXT("name"), Name);
+		bSetAField &= Json::TryGetField(InJson, TEXT("description"), Description);
+		bSetAField &= Json::TryGetField(InJson, TEXT("headers"), Headers);
+		bSetAField &= Json::TryGetField(InJson, TEXT("content"), Content);
+		return bSetAField;
+	}
+
+	bool FSchemaObject::FromJson(const TSharedRef<FJsonObject>& InJson)
+	{
+		bool bSetAField = Super::FromJson(InJson);
+		
+		bSetAField &= Json::TryGetField(InJson, TEXT("type"), Type);
+		bSetAField &= Json::TryGetField(InJson, TEXT("format"), Format);
+		bSetAField &= Json::TryGetField(InJson, TEXT("default"), Default);
+		bSetAField &= Json::TryGetField(InJson, TEXT("allOf"), AllOf);
+		bSetAField &= Json::TryGetField(InJson, TEXT("oneOf"), OneOf);
+		bSetAField &= Json::TryGetField(InJson, TEXT("anyOf"), AnyOf);
+		bSetAField &= Json::TryGetField(InJson, TEXT("not"), Not);
+
+		bSetAField &= Json::TryGetField(InJson, TEXT("title"), Title);
+		bSetAField &= Json::TryGetField(InJson, TEXT("multipleOf"), MultipleOf);
+		bSetAField &= Json::TryGetField(InJson, TEXT("minimum"), Minimum);
+		bSetAField &= Json::TryGetField(InJson, TEXT("exclusiveMinimum"), bExclusiveMinimum);
+		bSetAField &= Json::TryGetField(InJson, TEXT("maximum"), Maximum);
+		bSetAField &= Json::TryGetField(InJson, TEXT("exclusiveMaximum"), bExclusiveMaximum);
+		bSetAField &= Json::TryGetField(InJson, TEXT("minLength"), MinLength);
+		bSetAField &= Json::TryGetField(InJson, TEXT("maxLength"), MaxLength);
+		bSetAField &= Json::TryGetField(InJson, TEXT("pattern"), Pattern);
+		bSetAField &= Json::TryGetField(InJson, TEXT("minItems"), MinItems);
+		bSetAField &= Json::TryGetField(InJson, TEXT("maxItems"), MaxItems);
+		bSetAField &= Json::TryGetField(InJson, TEXT("uniqueItems"), bUniqueItems);
+		bSetAField &= Json::TryGetField(InJson, TEXT("minProperties"), MinProperties);
+		bSetAField &= Json::TryGetField(InJson, TEXT("maxProperties"), MaxProperties);
+		bSetAField &= Json::TryGetField(InJson, TEXT("enum"), Enum);
+		bSetAField &= Json::TryGetField(InJson, TEXT("items"), Items);
+		bSetAField &= Json::TryGetField(InJson, TEXT("nullable"), bNullable);
+		bSetAField &= Json::TryGetField(InJson, TEXT("discriminator"), Discriminator);
+		bSetAField &= Json::TryGetField(InJson, TEXT("readOnly"), bReadOnly);
+		bSetAField &= Json::TryGetField(InJson, TEXT("writeOnly"), bWriteOnly);
+		bSetAField &= Json::TryGetField(InJson, TEXT("xml"), Xml);
+		bSetAField &= Json::TryGetField(InJson, TEXT("externalDocs"), ExternalDocs);
+		// bSetAField &= Json::TryGetField(InJson, TEXT("example"), Example);
+
+		bSetAField &= Json::TryGetField(InJson, TEXT("properties"), Properties);
+		
+		return bSetAField;
+	}
+
+	bool FSecuritySchemeObject::FromJson(const TSharedRef<FJsonObject>& InJson)
+	{
+		bool bSetAField = true;
+		bSetAField &= Json::TryGetField(InJson, TEXT("type"), Type);
+		bSetAField &= Json::TryGetField(InJson, TEXT("description"), Description);
+		bSetAField &= Json::TryGetField(InJson, TEXT("name"), Name);
+		bSetAField &= Json::TryGetField(InJson, TEXT("in"), In);
+		bSetAField &= Json::TryGetField(InJson, TEXT("scheme"), Scheme);
+		bSetAField &= Json::TryGetField(InJson, TEXT("bearerFormat"), BearerFormat);
+		bSetAField &= Json::TryGetField(InJson, TEXT("flows"), Flows);
+		bSetAField &= Json::TryGetField(InJson, TEXT("openIdConnectUrl"), OpenIdConnectUrl);		
+		return bSetAField;
+	}
+
+	bool FServerObject::FromJson(const TSharedRef<FJsonObject>& InJson)
+	{
+		bool bSetAField = true;
+		bSetAField &= Json::TryGetField(InJson, TEXT("url"), Url);
+		bSetAField &= Json::TryGetField(InJson, TEXT("description"), Description);
+		bSetAField &= Json::TryGetField(InJson, TEXT("variables"), Variables);
+		return bSetAField;
+	}
+
+	bool FServerVariableObject::FromJson(const TSharedRef<FJsonObject>& InJson)
+	{
+		bool bSetAField = true;
+		bSetAField &= Json::TryGetField(InJson, TEXT("enum"), Enum);
+		bSetAField &= Json::TryGetField(InJson, TEXT("default"), Default);
+		bSetAField &= Json::TryGetField(InJson, TEXT("description"), Description);
+		return bSetAField;
+	}
+
+	bool FTagObject::FromJson(const TSharedRef<FJsonObject>& InJson)
+	{
+		bool bSetAField = true;
+		bSetAField &= Json::TryGetField(InJson, TEXT("name"), Name);
+		bSetAField &= Json::TryGetField(InJson, TEXT("description"), Description);
+		bSetAField &= Json::TryGetField(InJson, TEXT("externalDocs"), ExternalDocs);
+		return bSetAField;
+	}
+
+	bool FXMLObject::FromJson(const TSharedRef<FJsonObject>& InJson)
+	{
+		bool bSetAField = true;
+		bSetAField &= Json::TryGetField(InJson, TEXT("name"), Name);
+		bSetAField &= Json::TryGetField(InJson, TEXT("namespace"), Namespace);
+		bSetAField &= Json::TryGetField(InJson, TEXT("prefix"), Prefix);
+		bSetAField &= Json::TryGetField(InJson, TEXT("attribute"), bAttribute);
+		bSetAField &= Json::TryGetField(InJson, TEXT("wrapped"), bWrapped);
+		return bSetAField;
+	}
+
+	// https://swagger.io/specification/#path-item-object
+	bool FPathItemObject::FromJson(const TSharedRef<FJsonObject>& InJson)
+	{
+		// Describes the operations available on a single path. A Path Item MAY be empty,
+		// due to ACL constraints. The path itself is still exposed to the documentation viewer
+		// but they will not know which operations and parameters are available.
+		bool bSetAField = true;
+		bSetAField &= Json::TryGetField(InJson, TEXT("summary"), Summary);
+		bSetAField &= Json::TryGetField(InJson, TEXT("description"), Description);
+		bSetAField &= Json::TryGetField(InJson, TEXT("get"), Get);
+		bSetAField &= Json::TryGetField(InJson, TEXT("put"), Put);
+		bSetAField &= Json::TryGetField(InJson, TEXT("post"), Post);
+		bSetAField &= Json::TryGetField(InJson, TEXT("delete"), Delete);
+		bSetAField &= Json::TryGetField(InJson, TEXT("options"), Options);
+		bSetAField &= Json::TryGetField(InJson, TEXT("head"), Head);
+		bSetAField &= Json::TryGetField(InJson, TEXT("patch"), Patch);
+		bSetAField &= Json::TryGetField(InJson, TEXT("trace"), Trace);
+		return bSetAField;
+	}
+}
