WebAPI fix

#cis fix

diff --git a/Engine/Plugins/Web/WebAPI/Source/WebAPIEditor/Public/WebApiJsonUtilities.h b/Engine/Plugins/Web/WebAPI/Source/WebAPIEditor/Public/WebApiJsonUtilities.h
--- a/Engine/Plugins/Web/WebAPI/Source/WebAPIEditor/Public/WebApiJsonUtilities.h
+++ b/Engine/Plugins/Web/WebAPI/Source/WebAPIEditor/Public/WebApiJsonUtilities.h
@@ -327,8 +327,8 @@ namespace UE::Json
 		{
 			if(InNameToValue.IsEmpty())
 			{
-				UE_LOG(LogWebAPI, Error, TEXT("Tried to convert enum from string, but no lookup was provided."));
-				return bResult;  
+				UE_LOG(LogTemp, Error, TEXT("Tried to convert enum from string, but no lookup was provided."));
+				return bResult;
 			}
 
 			FString StrValue;
diff --git a/Engine/Plugins/Web/WebAPI/Source/WebAPIOpenAPI/Private/V2/WebAPISwaggerConverter.cpp b/Engine/Plugins/Web/WebAPI/Source/WebAPIOpenAPI/Private/V2/WebAPISwaggerConverter.cpp
--- a/Engine/Plugins/Web/WebAPI/Source/WebAPIOpenAPI/Private/V2/WebAPISwaggerConverter.cpp
+++ b/Engine/Plugins/Web/WebAPI/Source/WebAPIOpenAPI/Private/V2/WebAPISwaggerConverter.cpp
@@ -909,7 +909,7 @@ namespace UE
 				const TObjectPtr<UWebAPIModelBase> ModelBase = Cast<UWebAPIModelBase>(OutParameter);
 				if (!ConvertModelBase(InParameter, ModelBase))
 				{
-					return nullptr;
+					return false;
 				}
 
 				if(!OutParameter->Type.TypeInfo->bIsBuiltinType)
diff --git a/Engine/Plugins/Web/WebAPI/Source/WebAPIOpenAPI/Private/V3/WebAPIOpenAPIConverter.cpp b/Engine/Plugins/Web/WebAPI/Source/WebAPIOpenAPI/Private/V3/WebAPIOpenAPIConverter.cpp
--- a/Engine/Plugins/Web/WebAPI/Source/WebAPIOpenAPI/Private/V3/WebAPIOpenAPIConverter.cpp
+++ b/Engine/Plugins/Web/WebAPI/Source/WebAPIOpenAPI/Private/V3/WebAPIOpenAPIConverter.cpp
@@ -1036,7 +1036,7 @@ namespace UE
 				const TObjectPtr<UWebAPIModelBase> ModelBase = Cast<UWebAPIModelBase>(OutParameter);
 				if (!ConvertModelBase(InParameter, ModelBase))
 				{
-					return nullptr;
+					return false;
 				}
 
 				if(!OutParameter->Type.TypeInfo->bIsBuiltinType)
