WebAPI initial submission redux

diff --git a/Engine/Plugins/Web/WebAPI/Source/WebAPI/Private/Security/WebAPIAuthentication.cpp b/Engine/Plugins/Web/WebAPI/Source/WebAPI/Private/Security/WebAPIAuthentication.cpp
new file mode 100644
--- /dev/null
+++ b/Engine/Plugins/Web/WebAPI/Source/WebAPI/Private/Security/WebAPIAuthentication.cpp
@@ -0,0 +1,163 @@
+﻿// Copyright Epic Games, Inc. All Rights Reserved.
+
+#include "Security/WebAPIAuthentication.h"
+
+#include "WebAPIDeveloperSettings.h"
+#include "WebAPILog.h"
+#include "WebAPISubsystem.h"
+#include "Engine/Engine.h"
+#include "Interfaces/IHttpRequest.h"
+#include "Serialization/JsonSerializer.h"
+
+UWebAPIOAuthSettings::UWebAPIOAuthSettings()
+{
+	SchemeName = TEXT("OAuth");
+}
+
+bool UWebAPIOAuthSettings::IsValid() const
+{
+	return !ClientId.IsEmpty() && !ClientSecret.IsEmpty();
+}
+
+bool FWebAPIOAuthSchemeHandler::HandleHttpRequest(TSharedPtr<IHttpRequest> InRequest, UWebAPIDeveloperSettings* InSettings)
+{
+	check(InSettings);
+
+	// @todo: cache!
+	UWebAPIOAuthSettings* OAuthSettings = nullptr;
+	InSettings->AuthenticationSettings.FindItemByClass<UWebAPIOAuthSettings>(&OAuthSettings);	
+	check(OAuthSettings);
+
+	if(!ensureMsgf(OAuthSettings->IsValid(), TEXT("Authentication settings are missing one or more required properties.")))
+	{
+		return false;
+	}
+
+	if(OAuthSettings->AccessToken.IsEmpty())
+	{
+		return false;
+	}
+
+	InRequest->SetHeader(TEXT("Authorization"), FString::Printf(TEXT("%s %s"), *OAuthSettings->TokenType, *OAuthSettings->AccessToken));
+
+	return true;
+}
+
+bool FWebAPIOAuthSchemeHandler::HandleHttpResponse(EHttpResponseCodes::Type InResponseCode, TSharedPtr<IHttpResponse, ESPMode::ThreadSafe> InResponse, bool bInWasSuccessful, UWebAPIDeveloperSettings* InSettings)
+{
+	check(InSettings);
+
+	// @todo: cache!
+	UWebAPIOAuthSettings* OAuthSettings = nullptr;
+	InSettings->AuthenticationSettings.FindItemByClass<UWebAPIOAuthSettings>(&OAuthSettings);	
+	check(OAuthSettings);
+	
+	if(!ensureMsgf(OAuthSettings->IsValid(), TEXT("Authentication settings are missing one or more required properties.")))
+	{
+		return false;
+	}
+
+	GEngine->GetEngineSubsystem<UWebAPISubsystem>()->MakeHttpRequest(TEXT("POST"), [&, OAuthSettings](const TSharedPtr<IHttpRequest>& InRequest)
+	{
+		//Microsoft login rest URL
+		FStringFormatNamedArguments UrlArgs;
+		UrlArgs.Add(TEXT("TenantId"), OAuthSettings->TenantId);
+		const FString Url = FString::Format(*OAuthSettings->AuthenticationServer, UrlArgs);
+
+		// @todo: move!
+		FString StringPayload = "grant_type=client_credentials";
+		StringPayload.Append("&client_id=" + OAuthSettings->ClientId);
+		StringPayload.Append("&client_secret=" + OAuthSettings->ClientSecret);
+		StringPayload.Append("&resource=https://digitaltwins.azure.net");
+
+		InRequest->SetURL(Url);
+		InRequest->SetContentAsString(StringPayload);
+
+		TMap<FString, FString> Headers;
+		// FInternetAddr
+		FString SchemeName;
+		FParse::SchemeNameFromURI(*Url, SchemeName);
+		FString Host = Url.Replace(*(SchemeName + TEXT("://")), TEXT(""));
+		
+		int32 DelimiterIdx = -1;
+		if(Host.FindChar(TEXT('/'), DelimiterIdx))
+		{
+			Host = Host.Left(DelimiterIdx);
+		}
+
+		Headers.Add("Host", Host);
+		Headers.Add("Content-Type", "application/x-www-form-urlencoded");
+
+		for (const TPair<FString, FString>& Header : Headers)
+		{
+			InRequest->SetHeader(Header.Key, Header.Value);
+		}
+	})
+	.Next([this, OAuthSettings](const TTuple<FHttpResponsePtr, bool>& InResponse)
+	{
+		// Request failed
+		if(!InResponse.Get<bool>())
+		{
+			return false;
+		}
+
+		const FHttpResponsePtr Response = InResponse.Get<FHttpResponsePtr>();
+
+		FString Message;
+		FString AccessTkn;
+		bool bSuccess = false;
+
+		if (Response.IsValid())
+		{
+			const int32 ResponseCode = Response->GetResponseCode();
+			Message = Response->GetContentAsString();
+
+			if (EHttpResponseCodes::IsOk(ResponseCode))
+			{
+				const TSharedRef<TJsonReader<>> JsonReader = TJsonReaderFactory<>::Create(Response->GetContentAsString());
+				TSharedPtr<FJsonObject> JsonObject;
+				if(FJsonSerializer::Deserialize(JsonReader, JsonObject))
+				{
+					const int32 UnixTimeExpire = JsonObject->GetNumberField("expires_on");
+
+					UE_LOG(LogWebAPI, Display, TEXT("Generate token Response success"));
+					OAuthSettings->TokenType = JsonObject->GetStringField("token_type");
+					OAuthSettings->AccessToken = JsonObject->GetStringField("access_token");
+					OAuthSettings->ExpiresOn =  FDateTime::FromUnixTimestamp(UnixTimeExpire);
+					OAuthSettings->SaveConfig();
+					bSuccess = true;
+				}
+				else
+				{
+					UE_LOG(LogWebAPI, Warning, TEXT("Deserialize JSON failed"));
+					UE_LOG(LogWebAPI, Error, TEXT("Authentication failed: Deserialize JSON Response token failed"));
+					Message = "Deserialize JSON Response token failed:" + Message;
+				}
+			}
+			else
+				{
+				UE_LOG(LogWebAPI, Error, TEXT("Authentication failed: Response not valid"));
+				Message = "Response code not valid:" + Message;
+			}
+		}
+		else
+		{
+			UE_LOG(LogWebAPI, Error, TEXT("Authentication failed: Generate token Response not valid"));
+			Message = "Response is null";
+		}
+
+		// @todo:
+		//Redo all the requests that arrived while fetching new token
+		/*
+		if (bUpdatingToken && bSuccess)
+		{
+			UpdatingToken = false;
+			RedoRequests();
+		}
+		*/
+
+		return bSuccess;
+	});
+
+	return true;
+}
diff --git a/Engine/Plugins/Web/WebAPI/Source/WebAPI/Public/WebAPIMessageResponse.h b/Engine/Plugins/Web/WebAPI/Source/WebAPI/Public/WebAPIMessageResponse.h
new file mode 100644
--- /dev/null
+++ b/Engine/Plugins/Web/WebAPI/Source/WebAPI/Public/WebAPIMessageResponse.h
@@ -0,0 +1,19 @@
+﻿// Copyright Epic Games, Inc. All Rights Reserved.
+
+#pragma once
+
+#include "CoreMinimal.h"
+#include "WebAPIMessageResponse.generated.h"
+
+/**
+ * 
+ */
+USTRUCT(BlueprintType)
+struct WEBAPI_API FWebAPIMessageResponse
+{
+	GENERATED_BODY()
+
+	virtual ~FWebAPIMessageResponse() = default;
+
+	virtual const FText& GetMessage() const PURE_VIRTUAL(FWebAPIMessageResponse::GetMessage, return FText::GetEmpty(););
+};
diff --git a/Engine/Plugins/Web/WebAPI/Source/WebAPIBlueprintGraph/Private/K2Node_WebAPIOperation.cpp b/Engine/Plugins/Web/WebAPI/Source/WebAPIBlueprintGraph/Private/K2Node_WebAPIOperation.cpp
new file mode 100644
--- /dev/null
+++ b/Engine/Plugins/Web/WebAPI/Source/WebAPIBlueprintGraph/Private/K2Node_WebAPIOperation.cpp
@@ -0,0 +1,379 @@
+﻿// Copyright Epic Games, Inc. All Rights Reserved.
+
+#include "K2Node_WebAPIOperation.h"
+
+#include "BlueprintActionDatabaseRegistrar.h"
+#include "BlueprintTypePromotion.h"
+#include "EdGraphSchema_K2.h"
+#include "K2Node_StructOperation.h"
+#include "KismetCompiler.h"
+#include "ScopedTransaction.h"
+#include "ToolMenu.h"
+#include "WebAPIOperationObject.h"
+#include "Algo/AllOf.h"
+#include "EdGraph/EdGraphNode.h"
+#include "EdGraph/EdGraphPin.h"
+#include "Kismet2/BlueprintEditorUtils.h"
+
+#define LOCTEXT_NAMESPACE "K2Node_WebAPIAsyncOperation"
+
+struct FGetPinName
+{
+	static const FName& GetMessagePinName() { static const FName MessagePinName(TEXT("Message")); return MessagePinName; }
+	static const FName& GetSuccessPinName() { static const FName SuccessPinName(TEXT("OnSuccess")); return SuccessPinName; }
+	static const FName& GetErrorPinName() { static const FName ErrorPinName(TEXT("OnError")); return ErrorPinName; }
+};
+
+UK2Node_WebAPIOperation::UK2Node_WebAPIOperation()
+{
+	ProxyActivateFunctionName = NAME_None;
+
+	// @todo:
+	// 1. on new node spawned, break Request (input) pin by default: FBlueprintEditor::OnSplitStructPin()
+	// 2. option to convert to/from callback styles: UK2Node_DynamicCast::SetPurity(bool bNewPurity)
+	// 3. collapse output responses that inherit MessageResponse to single message output pin
+}
+
+void UK2Node_WebAPIOperation::AllocateDefaultPins()
+{
+	Super::AllocateDefaultPins();
+}
+
+FText UK2Node_WebAPIOperation::GetNodeTitle(ENodeTitleType::Type TitleType) const
+{
+	FText FunctionName = Super::GetNodeTitle(TitleType);
+	FText NamespaceName;
+	FText ServiceName;
+	
+	if (const UFunction* Function = GetFactoryFunction())
+	{
+		NamespaceName = Function->GetOuterUClass()->GetMetaDataText(TEXT("Namespace"));
+		ServiceName = Function->GetOuterUClass()->GetMetaDataText(TEXT("Service"));
+	}
+	
+	if(TitleType == ENodeTitleType::FullTitle)
+	{
+		FFormatNamedArguments Args;
+		Args.Add(TEXT("FunctionName"), FunctionName);
+		Args.Add(TEXT("NamespaceName"), NamespaceName);
+		Args.Add(TEXT("ServiceName"), ServiceName);
+
+		if (NamespaceName.IsEmpty() && ServiceName.IsEmpty())
+		{
+			return FText::Format(LOCTEXT("NodeTitle", "{FunctionName}"), Args);
+		}
+		else if (NamespaceName.IsEmpty())
+		{
+			return FText::Format(LOCTEXT("NodeTitle_WithNamespace", "{FunctionName}\n{ServiceName}"), Args);
+		}
+		else if (ServiceName.IsEmpty())
+		{
+			return FText::Format(LOCTEXT("NodeTitle_WithService", "{FunctionName}\n{NamespaceName}"), Args);
+		}
+		else
+		{
+			return FText::Format(LOCTEXT("NodeTitle_WithNamespaceAndService", "{FunctionName}\n{NamespaceName}: {ServiceName}"), Args);
+		}		
+	}
+	else
+	{
+		return FunctionName;
+	}
+}
+
+void UK2Node_WebAPIOperation::GetNodeContextMenuActions(UToolMenu* Menu, UGraphNodeContextMenuContext* Context) const
+{
+	Super::GetNodeContextMenuActions(Menu, Context);
+
+	if (!Context->bIsDebugging)
+	{
+		{
+			FText MenuEntryTitle = LOCTEXT("MakeCallbackTitle", "Convert to Callback responses");
+			FText MenuEntryTooltip = LOCTEXT("MakeCallbackTooltip", "Removes the execution pins and instead adds callbacks.");
+
+			bool bCanToggleAsyncType = true;
+			auto CanExecuteAsyncTypeToggle = [](const bool bInCanToggleAsyncType)->bool
+			{
+				return bInCanToggleAsyncType;
+			};
+
+			if (OperationAsyncType == EWebAPIOperationAsyncType::Callback)
+			{
+				MenuEntryTitle = LOCTEXT("MakeLatentActionTitle", "Convert to Latent Action");
+				MenuEntryTooltip = LOCTEXT("MakeLatentActionTooltip", "Adds in branching execution pins so that you can separatly handle the responses.");
+
+				const UEdGraphSchema_K2* K2Schema = Cast<UEdGraphSchema_K2>(GetSchema());
+				check(K2Schema != nullptr);
+
+				bCanToggleAsyncType = K2Schema->DoesGraphSupportImpureFunctions(GetGraph()); // @todo 
+				if (!bCanToggleAsyncType)
+				{
+					MenuEntryTooltip = LOCTEXT("CannotMakeLatentActionTooltip", "This graph does not support latent actions.");
+				}
+			}
+
+			FToolMenuSection& Section = Menu->AddSection("K2NodeWebAPIOperation", LOCTEXT("AsyncTypeHeader", "Async Type"));
+			Section.AddMenuEntry(
+				"ToggleAsyncType",
+				MenuEntryTitle,
+				MenuEntryTooltip,
+				FSlateIcon(),
+				FUIAction(
+					FExecuteAction::CreateUObject(const_cast<UK2Node_WebAPIOperation*>(this), &UK2Node_WebAPIOperation::ToggleAsyncType),
+					FCanExecuteAction::CreateStatic(CanExecuteAsyncTypeToggle, bCanToggleAsyncType),
+					FIsActionChecked()
+				)
+			);
+		}
+	}
+}
+
+bool UK2Node_WebAPIOperation::CanPasteHere(const UEdGraph* TargetGraph) const
+{
+	// @todo: unlike default async node, you can paste into a function (not event graph) and this should convert to callback style
+	return Super::CanPasteHere(TargetGraph);
+}
+
+void UK2Node_WebAPIOperation::PostPlacedNewNode()
+{
+	Super::PostPlacedNewNode();
+
+	PostReconstructNode();
+}
+
+void UK2Node_WebAPIOperation::PostPasteNode()
+{
+	Super::PostPasteNode();
+
+	// @todo: handle conversion to/from callback style
+}
+
+void UK2Node_WebAPIOperation::GetMenuActions(FBlueprintActionDatabaseRegistrar& ActionRegistrar) const
+{
+	struct GetMenuActions_Utils
+	{
+		static void SetNodeFunc(UEdGraphNode* NewNode, bool /*bIsTemplateNode*/, TWeakObjectPtr<UFunction> FunctionPtr)
+		{
+			UK2Node_WebAPIOperation* AsyncTaskNode = CastChecked<UK2Node_WebAPIOperation>(NewNode);
+			if (FunctionPtr.IsValid())
+			{
+				const UFunction* Func = FunctionPtr.Get();
+				FObjectProperty* ReturnProp = CastFieldChecked<FObjectProperty>(Func->GetReturnProperty());
+
+				AsyncTaskNode->ProxyFactoryFunctionName = Func->GetFName();
+				AsyncTaskNode->ProxyFactoryClass        = Func->GetOuterUClass();
+				AsyncTaskNode->ProxyClass               = ReturnProp->PropertyClass;				
+			}
+		}
+	};
+
+	UClass* NodeClass = GetClass();
+	ActionRegistrar.RegisterClassFactoryActions<UWebAPIOperationObject>(FBlueprintActionDatabaseRegistrar::FMakeFuncSpawnerDelegate::CreateLambda([NodeClass](const UFunction* FactoryFunc)->UBlueprintNodeSpawner*
+	{
+		UBlueprintNodeSpawner* NodeSpawner = UBlueprintFunctionNodeSpawner::Create(FactoryFunc);
+		check(NodeSpawner != nullptr);
+		
+		NodeSpawner->NodeClass = NodeClass;
+
+		const TWeakObjectPtr<UFunction> FunctionPtr = MakeWeakObjectPtr(const_cast<UFunction*>(FactoryFunc));
+		NodeSpawner->CustomizeNodeDelegate = UBlueprintNodeSpawner::FCustomizeNodeDelegate::CreateStatic(GetMenuActions_Utils::SetNodeFunc, FunctionPtr);
+
+		return NodeSpawner;
+	}));
+}
+
+void UK2Node_WebAPIOperation::PostReconstructNode()
+{
+	Super::PostReconstructNode();
+
+	if(IsValid())
+	{
+		SplitRequestPins(GetRequestPins());
+		SplitResponsePins(GetResponsePins());
+		SplitResponsePins(GetErrorResponsePins());
+	}
+}
+
+void UK2Node_WebAPIOperation::SetAsyncType(EWebAPIOperationAsyncType InAsyncType)
+{
+	if (InAsyncType != OperationAsyncType)
+	{
+		OperationAsyncType = InAsyncType;
+
+		const bool bHasBeenConstructed = (Pins.Num() > 0);
+		if (bHasBeenConstructed)
+		{
+			ReconstructNode();
+		}
+	}
+}
+
+bool UK2Node_WebAPIOperation::IsValid() const
+{
+	return ProxyFactoryFunctionName != NAME_None
+		&& ProxyClass != nullptr
+		&& ((bHasCompilerMessage && ErrorType >= EMessageSeverity::Info) || !bHasCompilerMessage);
+}
+
+UEdGraphPin* UK2Node_WebAPIOperation::FindPinContainingName(
+	const EEdGraphPinDirection& InPinDirection,
+	const FString& InName) const
+{
+	return FindPinByPredicate([&InPinDirection, &InName](const UEdGraphPin* InPin)
+	{
+		return InPin->Direction == InPinDirection
+			&& InPin->PinType.PinCategory != UEdGraphSchema_K2::PC_Exec
+			&& InPin->GetName().Contains(InName);
+	});
+}
+
+TArray<UEdGraphPin*> UK2Node_WebAPIOperation::FindPinsContainingName(
+	const EEdGraphPinDirection& InPinDirection,
+	const FString& InName) const
+{
+	return Pins.FilterByPredicate([&InPinDirection, InName](const UEdGraphPin* InPin)
+	{
+		return InPin->Direction == InPinDirection
+			&& InPin->PinType.PinCategory != UEdGraphSchema_K2::PC_Exec
+			&& InPin->GetName().Contains(InName);
+	});
+}
+
+TArray<UEdGraphPin*> UK2Node_WebAPIOperation::GetRequestPins() const
+{
+	TArray<UEdGraphPin*> RequestPins = FindPinsContainingName(EGPD_Input, TEXT("Request"));
+	ensureMsgf(!RequestPins.IsEmpty(), TEXT("The function must contain a parameter with \"Request\" in the name."));	
+	return RequestPins;
+}
+
+TArray<UEdGraphPin*> UK2Node_WebAPIOperation::GetResponsePins() const
+{
+	auto ResponsePins = FindPinsContainingName(EGPD_Output, TEXT("Response"));
+	ensureMsgf(!ResponsePins.IsEmpty(), TEXT("The Operation must contain a delegate with a parameter containing \"Response\" in the name."));	
+	return ResponsePins;
+}
+
+TArray<UEdGraphPin*> UK2Node_WebAPIOperation::GetErrorResponsePins() const
+{	
+	auto ErrorResponsePins = FindPinsContainingName(EGPD_Output, TEXT("Error"));
+	ensureMsgf(!ErrorResponsePins.IsEmpty(), TEXT("The Operation must contain a delegate with a parameter containing \"Error\" in the name."));	
+	return ErrorResponsePins;
+}
+
+// @todo: recursively split where struct is basically a wrapper for another struct
+void UK2Node_WebAPIOperation::SplitRequestPins(const TArray<UEdGraphPin*>& InPins) const
+{
+	// If UWebAPIOperation base class, the pin won't always be valid, so ignore.
+	if(!ensureAlways(!InPins.IsEmpty() && Algo::AllOf(InPins, [](const UEdGraphPin* InPin) { return InPin; })))
+	{
+		return;
+	}
+
+	for(UEdGraphPin* Pin : InPins)
+	{
+		// Split if not already
+		if (Pin->SubPins.Num() == 0)
+		{
+			if(CanSplitPin(Pin))
+			{
+				const UEdGraphSchema_K2* K2Schema = GetDefault<UEdGraphSchema_K2>();
+				K2Schema->SplitPin(Pin);
+			}
+		}
+
+		// Removes the prefix from the auto-named split pins 
+		if (Pin->SubPins.Num() > 0)
+		{
+			for (UEdGraphPin* SubPin : Pin->SubPins)
+			{
+				FString SubPinName = SubPin->PinName.ToString();
+				CleanupPinNameInline(SubPinName);
+
+				//SubPin->Modify();			
+				SubPin->PinFriendlyName = FText::FromString(SubPinName);
+			}
+		}
+	}
+}
+
+// Splits if the Response contains a single property
+void UK2Node_WebAPIOperation::SplitResponsePins(const TArray<UEdGraphPin*>& InPins) const
+{
+	// Compilation errors or stale nodes should be warnings, not errors
+	if(!ensureAlways(!InPins.IsEmpty() && Algo::AllOf(InPins, [](const UEdGraphPin* InPin) { return InPin; })))
+	{
+		return;	
+	}
+
+	for(UEdGraphPin* Pin : InPins)
+	{
+		if (const UScriptStruct* PinStructType = Cast<UScriptStruct>(Pin->PinType.PinSubCategoryObject.Get()))
+		{
+			uint32 PropertyNum = 0;
+			
+			// Split if not already
+			if (Pin->SubPins.Num() == 0 && CanSplitPin(Pin))
+			{
+				for (TFieldIterator<const FProperty> PropertyIt(PinStructType, EFieldIteratorFlags::ExcludeSuper, EFieldIteratorFlags::ExcludeDeprecated); PropertyIt; ++PropertyIt)
+				{
+					if (PropertyIt->HasAnyPropertyFlags(CPF_Transient))
+					{
+						continue;
+					}
+
+					PropertyNum++;
+				}
+
+				if(PropertyNum == 1)
+				{
+					// Split, result should be single pin (cause it only contained a single property!)
+					const UEdGraphSchema_K2* K2Schema = GetDefault<UEdGraphSchema_K2>();
+					K2Schema->SplitPin(Pin);
+				}
+			}
+
+			// Removes the prefix from the auto-named split pins 
+			if (Pin->SubPins.Num() > 0)
+			{
+				for (UEdGraphPin* SubPin : Pin->SubPins)
+				{
+					FString SubPinName = SubPin->PinName.ToString();
+					CleanupPinNameInline(SubPinName);
+
+					//SubPin->Modify();			
+					SubPin->PinFriendlyName = FText::FromString(SubPinName);
+				}
+			}
+		}
+	}
+}
+
+void UK2Node_WebAPIOperation::CleanupPinNameInline(FString& InPinName)
+{
+	int32 UnderscorePosition = -1;
+	if(InPinName.FindLastChar(TCHAR('_'), UnderscorePosition))
+	{
+		InPinName = InPinName.RightChop(UnderscorePosition + 1);
+	}
+}
+
+void UK2Node_WebAPIOperation::ToggleAsyncType()
+{
+	const FText TransactionTitle =
+		OperationAsyncType == EWebAPIOperationAsyncType::Callback
+		? LOCTEXT("ToggleAsyncTypeToLatentAction", "Convert to Latent Action")
+		: LOCTEXT("ToggleAsyncTypeToCallback", "Convert to Callback responses");
+	const FScopedTransaction Transaction(TransactionTitle);
+	Modify();
+
+	SetAsyncType(StaticCast<EWebAPIOperationAsyncType>(FMath::Abs(StaticCast<int8>(OperationAsyncType) - 1)));
+}
+
+bool UK2Node_WebAPIOperation::ReconnectExecPins(TArray<UEdGraphPin*>& OldPins)
+{
+	// @todo: handle when converting to/from callback style
+	
+	return true;
+}
+
+#undef LOCTEXT_NAMESPACE
diff --git a/Engine/Plugins/Web/WebAPI/Source/WebAPIBlueprintGraph/Private/K2Node_WebAPIOperation.h b/Engine/Plugins/Web/WebAPI/Source/WebAPIBlueprintGraph/Private/K2Node_WebAPIOperation.h
new file mode 100644
--- /dev/null
+++ b/Engine/Plugins/Web/WebAPI/Source/WebAPIBlueprintGraph/Private/K2Node_WebAPIOperation.h
@@ -0,0 +1,71 @@
+﻿// Copyright Epic Games, Inc. All Rights Reserved.
+
+#pragma once
+
+#include "CoreMinimal.h"
+#include "K2Node_BaseAsyncTask.h"
+
+#include "K2Node_WebAPIOperation.generated.h"
+
+class FBlueprintActionDatabaseRegistrar;
+
+UENUM()
+enum class EWebAPIOperationAsyncType : uint8
+{
+	LatentAction = 0,
+	Callback = 1
+};
+
+/**
+ * 
+ */
+UCLASS()
+class WEBAPIBLUEPRINTGRAPH_API UK2Node_WebAPIOperation
+	: public UK2Node_BaseAsyncTask
+{
+	GENERATED_BODY()
+
+public:
+	UK2Node_WebAPIOperation();
+
+	//~ Begin UEdGraphNode Interface.
+	virtual void AllocateDefaultPins() override;
+	virtual FText GetNodeTitle(ENodeTitleType::Type TitleType) const override;
+	virtual void GetNodeContextMenuActions(class UToolMenu* Menu, class UGraphNodeContextMenuContext* Context) const override;
+	virtual void PostPlacedNewNode() override;
+	virtual void PostPasteNode() override;
+	virtual bool CanPasteHere(const UEdGraph* TargetGraph) const override;
+	//~ End UEdGraphNode Interface.
+
+	// UK2Node interface
+	virtual void GetMenuActions(FBlueprintActionDatabaseRegistrar& ActionRegistrar) const override;
+	virtual void PostReconstructNode() override;
+	// End of UK2Node interface
+
+	void SetAsyncType(EWebAPIOperationAsyncType InAsyncType);
+
+	bool IsValid() const;
+
+private:
+	UEdGraphPin* FindPinContainingName(const EEdGraphPinDirection& InPinDirection, const FString& InName) const;
+	TArray<UEdGraphPin*> FindPinsContainingName(const EEdGraphPinDirection& InPinDirection, const FString& InName) const;
+	
+	TArray<UEdGraphPin*> GetRequestPins() const;
+	TArray<UEdGraphPin*> GetResponsePins() const;
+	TArray<UEdGraphPin*> GetErrorResponsePins() const;
+
+	void SplitRequestPins(const TArray<UEdGraphPin*>& InPins) const;	
+	void SplitResponsePins(const TArray<UEdGraphPin*>& InPins) const;
+
+	static void CleanupPinNameInline(FString& InPinName);
+
+	/** Flips the node's Async Type. */
+	void ToggleAsyncType();
+
+	/** Update exec pins when converting between LatentAction and Callback Async Types. */
+	bool ReconnectExecPins(TArray<UEdGraphPin*>& OldPins);
+
+	/** Latent Action is preferred, but not compatible when used in a function, so allow conversion between the two. */
+	UPROPERTY()
+	EWebAPIOperationAsyncType OperationAsyncType = EWebAPIOperationAsyncType::LatentAction;
+};
diff --git a/Engine/Plugins/Web/WebAPI/Source/WebAPIBlueprintGraph/Private/WebAPIBlueprintGraphModule.cpp b/Engine/Plugins/Web/WebAPI/Source/WebAPIBlueprintGraph/Private/WebAPIBlueprintGraphModule.cpp
new file mode 100644
--- /dev/null
+++ b/Engine/Plugins/Web/WebAPI/Source/WebAPIBlueprintGraph/Private/WebAPIBlueprintGraphModule.cpp
@@ -0,0 +1,5 @@
+// Copyright Epic Games, Inc. All Rights Reserved.
+
+#include "Modules/ModuleManager.h"
+
+IMPLEMENT_MODULE(FDefaultModuleImpl, WebAPIBlueprintGraph)
diff --git a/Engine/Plugins/Web/WebAPI/Source/WebAPIEditor/Private/Details/Widgets/SWebAPISchemaParameterRow.h b/Engine/Plugins/Web/WebAPI/Source/WebAPIEditor/Private/Details/Widgets/SWebAPISchemaParameterRow.h
new file mode 100644
--- /dev/null
+++ b/Engine/Plugins/Web/WebAPI/Source/WebAPIEditor/Private/Details/Widgets/SWebAPISchemaParameterRow.h
@@ -0,0 +1,19 @@
+﻿// Copyright Epic Games, Inc. All Rights Reserved.
+
+#pragma once
+
+#include "CoreMinimal.h"
+#include "SWebAPISchemaTreeTableRow.h"
+#include "Details/ViewModels/WebAPIParameterViewModel.h"
+
+class SWebAPISchemaParameterRow
+	: public SWebAPISchemaTreeTableRow<FWebAPIParameterViewModel>
+{
+public:
+	SLATE_BEGIN_ARGS(SWebAPISchemaParameterRow)
+	{ }
+	SLATE_END_ARGS()
+
+public:
+	void Construct(const FArguments& InArgs, const TSharedRef<FWebAPIParameterViewModel>& InViewModel, const TSharedRef<STableViewBase>& InOwnerTableView);
+};
diff --git a/Engine/Plugins/Web/WebAPI/Source/WebAPIEditor/Private/Dom/WebAPIModelParameter.cpp b/Engine/Plugins/Web/WebAPI/Source/WebAPIEditor/Private/Dom/WebAPIModelParameter.cpp
new file mode 100644
--- /dev/null
+++ b/Engine/Plugins/Web/WebAPI/Source/WebAPIEditor/Private/Dom/WebAPIModelParameter.cpp
@@ -0,0 +1 @@
+﻿// Copyright Epic Games, Inc. All Rights Reserved.
\ No newline at end of file
diff --git a/Engine/Plugins/Web/WebAPI/Source/WebAPIEditor/Private/Dom/WebAPIService.cpp b/Engine/Plugins/Web/WebAPI/Source/WebAPIEditor/Private/Dom/WebAPIService.cpp
new file mode 100644
--- /dev/null
+++ b/Engine/Plugins/Web/WebAPI/Source/WebAPIEditor/Private/Dom/WebAPIService.cpp
@@ -0,0 +1,45 @@
+﻿// Copyright Epic Games, Inc. All Rights Reserved.
+
+#include "Dom/WebAPIService.h"
+
+#include "Dom/WebAPIOperation.h"
+#include "Dom/WebAPIType.h"
+
+#define LOCTEXT_NAMESPACE "WebAPIService"
+
+void UWebAPIService::SetNamespace(const FString& InNamespace)
+{
+	Super::SetNamespace(InNamespace);
+	check(Name.HasTypeInfo());
+
+	if(Name.HasTypeInfo() && !Name.TypeInfo->bIsBuiltinType)
+	{
+		Name.TypeInfo->Namespace = InNamespace;
+	}
+}
+
+void UWebAPIService::Visit(const TFunctionRef<void(IWebAPISchemaObjectInterface*&)> InVisitor)
+{
+	Super::Visit(InVisitor);
+
+	for(const TObjectPtr<UWebAPIOperation>& Operation : Operations)
+	{
+		Operation->Visit(InVisitor);
+	}
+}
+
+#if WITH_EDITOR
+EDataValidationResult UWebAPIService::IsDataValid(TArray<FText>& ValidationErrors)
+{
+	EDataValidationResult ValidationResult = Super::IsDataValid(ValidationErrors);
+
+	for(const TObjectPtr<UWebAPIOperation>& Operation : Operations)
+	{
+		ValidationResult = CombineDataValidationResults(Operation->IsDataValid(ValidationErrors), ValidationResult);
+	}
+
+	return ValidationResult;
+}
+#endif // WITH_EDITOR
+
+#undef LOCTEXT_NAMESPACE
diff --git a/Engine/Plugins/Web/WebAPI/Source/WebAPIEditor/Private/WebAPIEditorStyle.cpp b/Engine/Plugins/Web/WebAPI/Source/WebAPIEditor/Private/WebAPIEditorStyle.cpp
new file mode 100644
--- /dev/null
+++ b/Engine/Plugins/Web/WebAPI/Source/WebAPIEditor/Private/WebAPIEditorStyle.cpp
@@ -0,0 +1,70 @@
+﻿// Copyright Epic Games, Inc. All Rights Reserved.
+
+#include "WebAPIEditorStyle.h"
+
+#include "Interfaces/IPluginManager.h"
+#include "Slate/SlateGameResources.h"
+#include "Styling/SlateStyleMacros.h"
+#include "Styling/SlateStyleRegistry.h"
+#include "Styling/StyleColors.h"
+
+namespace UE
+{
+	namespace WebAPI
+	{
+		namespace Private
+		{
+			const FVector2D Icon16x16(16.0f, 16.0f);
+			const FVector2D Icon20x20(20.0f, 20.0f);
+			const FVector2D Icon40x40(40.0f, 40.0f);
+			const FVector2D Icon64x64(64.0f, 64.0f);
+
+			const FName NAME_StyleName(TEXT("WebAPIEditorStyle"));
+
+			static TUniquePtr<FSlateStyleSet> StyleInstance;	
+		}
+	}
+}
+
+TUniquePtr<FSlateStyleSet> FWebAPIEditorStyle::StyleInstance = nullptr;
+
+using namespace UE::WebAPI;
+
+// This is to fix the issue that SlateStyleMacros like IMAGE_BRUSH look for RootToContentDir but StyleSet->RootToContentDir is how this style is set up
+#define RootToContentDir Private::StyleInstance->RootToContentDir
+#define RootToCoreContentDir Private::StyleInstance->RootToCoreContentDir
+
+void FWebAPIEditorStyle::Register()
+{
+	Private::StyleInstance = MakeUnique<FSlateStyleSet>(Private::NAME_StyleName);
+	Private::StyleInstance->SetCoreContentRoot(FPaths::EngineContentDir() / TEXT("Editor/Slate"));
+	Private::StyleInstance->SetContentRoot(IPluginManager::Get().FindPlugin(TEXT("WebAPI"))->GetContentDir() / TEXT("Editor/Slate"));
+
+	Private::StyleInstance->Set("WebAPI.TreeView.LabelBackground", new CORE_BOX_BRUSH("Common/GroupBorderLight", 6.f/18.f, FStyleColors::Dropdown));
+	
+	FSlateStyleRegistry::RegisterSlateStyle(*Private::StyleInstance.Get());
+}
+
+void FWebAPIEditorStyle::Unregister()
+{
+	FSlateStyleRegistry::UnRegisterSlateStyle(*Private::StyleInstance.Get());
+	Private::StyleInstance.Reset();
+}
+
+FName FWebAPIEditorStyle::GetStyleSetName()
+{
+	return Private::NAME_StyleName;
+}
+
+const ISlateStyle& FWebAPIEditorStyle::Get()
+{
+	check(Private::StyleInstance.IsValid());
+	return *Private::StyleInstance.Get();
+}
+
+#undef IMAGE_BRUSH
+#undef BOX_BRUSH
+#undef BORDER_BRUSH
+#undef DEFAULT_FONT
+#undef CORE_BOX_BRUSH
+#undef CORE_IMAGE_BRUSH
diff --git a/Engine/Plugins/Web/WebAPI/Source/WebAPIEditor/Private/WebAPIEditorWidgetUtilities.cpp b/Engine/Plugins/Web/WebAPI/Source/WebAPIEditor/Private/WebAPIEditorWidgetUtilities.cpp
new file mode 100644
--- /dev/null
+++ b/Engine/Plugins/Web/WebAPI/Source/WebAPIEditor/Private/WebAPIEditorWidgetUtilities.cpp
@@ -0,0 +1,4 @@
+﻿// Copyright Epic Games, Inc. All Rights Reserved.
+
+
+#include "WebAPIEditorWidgetUtilities.h"
diff --git a/Engine/Plugins/Web/WebAPI/Source/WebAPIEditor/Private/WebAPIEditorWidgetUtilities.h b/Engine/Plugins/Web/WebAPI/Source/WebAPIEditor/Private/WebAPIEditorWidgetUtilities.h
new file mode 100644
--- /dev/null
+++ b/Engine/Plugins/Web/WebAPI/Source/WebAPIEditor/Private/WebAPIEditorWidgetUtilities.h
@@ -0,0 +1,26 @@
+﻿// Copyright Epic Games, Inc. All Rights Reserved.
+
+#pragma once
+
+#include "CoreMinimal.h"
+#include "Styling/SlateColor.h"
+#include "UObject/Object.h"
+#include "WebAPIEditorWidgetUtilities.generated.h"
+
+/**
+ * 
+ */
+UCLASS()
+class WEBAPIEDITOR_API UWebAPIEditorWidgetUtilities : public UObject
+{
+	GENERATED_BODY()
+};
+
+
+class FWebAPIEditorWidgetsUtilities
+{
+public:
+	// @todo: verb type should be enum
+	static FSlateColor GetColorForVerb(uint8 InVerb) { return FSlateColor(FColor{255, 128, 64}); }
+	static FName GetIconNameForVerb(uint8 InVerb) { return NAME_None; }
+};
diff --git a/Engine/Plugins/Web/WebAPI/Source/WebAPIEditor/Public/WebAPIEditorLog.h b/Engine/Plugins/Web/WebAPI/Source/WebAPIEditor/Public/WebAPIEditorLog.h
new file mode 100644
--- /dev/null
+++ b/Engine/Plugins/Web/WebAPI/Source/WebAPIEditor/Public/WebAPIEditorLog.h
@@ -0,0 +1,7 @@
+// Copyright Epic Games, Inc. All Rights Reserved.
+
+#pragma once
+
+#include "CoreMinimal.h"
+
+DECLARE_LOG_CATEGORY_EXTERN(LogWebAPIEditor, Log, All);
diff --git a/Engine/Plugins/Web/WebAPI/Source/WebAPIOpenAPI/Private/Factories/WebAPIOpenAPIFactoryBase.cpp b/Engine/Plugins/Web/WebAPI/Source/WebAPIOpenAPI/Private/Factories/WebAPIOpenAPIFactoryBase.cpp
new file mode 100644
--- /dev/null
+++ b/Engine/Plugins/Web/WebAPI/Source/WebAPIOpenAPI/Private/Factories/WebAPIOpenAPIFactoryBase.cpp
@@ -0,0 +1,31 @@
+﻿// Copyright Epic Games, Inc. All Rights Reserved.
+
+#include "WebAPIOpenAPIFactoryBase.h"
+
+#define LOCTEXT_NAMESPACE "WebAPIOpenAPIFactoryBase"
+
+const FName UWebAPIOpenAPIFactoryBase::JsonFileType = FName(TEXT("json"));
+const FName UWebAPIOpenAPIFactoryBase::YamlFileType = FName(TEXT("yaml"));
+
+UWebAPIOpenAPIFactoryBase::UWebAPIOpenAPIFactoryBase()
+{
+	Formats.Add(TEXT("json;JavaScript Object Notation"));
+}
+
+TFuture<bool> UWebAPIOpenAPIFactoryBase::ImportWebAPI(UWebAPIDefinition* InDefinition, const FString& InFileName, const FString& InFileContents)
+{
+	// @todo: modify definition here
+	// write raw spec contents
+	// parse spec
+
+	// @todo: return false if json couldn't be parsed
+
+	return MakeFulfilledPromise<bool>(false).GetFuture();
+}
+
+bool UWebAPIOpenAPIFactoryBase::IsValidFileExtension(const FString& InFileExtension) const
+{
+	return InFileExtension.Equals(JsonFileType.ToString(), ESearchCase::IgnoreCase);
+}
+
+#undef LOCTEXT_NAMESPACE
diff --git a/Engine/Plugins/Web/WebAPI/Source/WebAPIOpenAPI/Private/Factories/WebAPIOpenAPIFactoryBase.h b/Engine/Plugins/Web/WebAPI/Source/WebAPIOpenAPI/Private/Factories/WebAPIOpenAPIFactoryBase.h
new file mode 100644
--- /dev/null
+++ b/Engine/Plugins/Web/WebAPI/Source/WebAPIOpenAPI/Private/Factories/WebAPIOpenAPIFactoryBase.h
@@ -0,0 +1,26 @@
+﻿// Copyright Epic Games, Inc. All Rights Reserved.
+
+#pragma once
+
+#include "CoreMinimal.h"
+#include "Factories/WebAPIDefinitionFactory.h"
+
+#include "WebAPIOpenAPIFactoryBase.generated.h"
+
+/** Common base factory between Swagger and OpenAPI factories. */
+UCLASS(Abstract)
+class WEBAPIOPENAPI_API UWebAPIOpenAPIFactoryBase
+	: public UWebAPIDefinitionFactory
+{
+	GENERATED_BODY()
+
+public:
+	UWebAPIOpenAPIFactoryBase();
+
+	virtual TFuture<bool> ImportWebAPI(UWebAPIDefinition* InDefinition, const FString& InFileName, const FString& InFileContents) override;
+	virtual bool IsValidFileExtension(const FString& InFileExtension) const override;
+
+protected:
+	static const FName JsonFileType;
+	static const FName YamlFileType;
+};
diff --git a/Engine/Plugins/Web/WebAPI/Source/WebAPIOpenAPI/Private/Tests/Samples/V3/ContactObjectExample.json b/Engine/Plugins/Web/WebAPI/Source/WebAPIOpenAPI/Private/Tests/Samples/V3/ContactObjectExample.json
new file mode 100644
--- /dev/null
+++ b/Engine/Plugins/Web/WebAPI/Source/WebAPIOpenAPI/Private/Tests/Samples/V3/ContactObjectExample.json
@@ -0,0 +1,5 @@
+﻿{
+  "name": "API Support",
+  "url": "http://www.example.com/support",
+  "email": "[REDACTED_EMAIL]"
+}
\ No newline at end of file
diff --git a/Engine/Plugins/Web/WebAPI/Source/WebAPIOpenAPI/Private/Tests/Samples/V3/removebg.json b/Engine/Plugins/Web/WebAPI/Source/WebAPIOpenAPI/Private/Tests/Samples/V3/removebg.json
new file mode 100644
--- /dev/null
+++ b/Engine/Plugins/Web/WebAPI/Source/WebAPIOpenAPI/Private/Tests/Samples/V3/removebg.json
@@ -0,0 +1,984 @@
+{
+  "openapi": "3.0.0",
+  "servers": [
+    {
+      "url": "https://api.remove.bg/v1.0"
+    }
+  ],
+  "info": {
+    "description": "Remove the background of any image",
+    "title": "Background Removal API",
+    "version": "1.0.0",
+    "x-apisguru-categories": [
+      "developer_tools"
+    ],
+    "x-logo": {
+      "url": "https://api.apis.guru/v2/cache/logo/https_remove.bg_apple-touch-icon.png"
+    },
+    "x-origin": [
+      {
+        "format": "openapi",
+        "url": "https://www.remove.bg/api/swagger.yaml",
+        "version": "3.0"
+      }
+    ],
+    "x-providerName": "remove.bg"
+  },
+  "security": [
+    {
+      "APIKeyHeader": []
+    }
+  ],
+  "paths": {
+    "/account": {
+      "get": {
+        "description": "Get the current credit balance and number of free API calls.\n\nNotes:\n\n* Balance changes may be delayed by several seconds. To locally keep track of your credit balance, you should therefore only call this endpoint initially (or e.g. when the user manually triggers a refresh), then use the `X-Credits-Charged` response header returned with each image processing response to adjust the local balance.\n\n* The \"*sizes*\" field is always \"all\", is deprecated and will soon be removed.\n",
+        "responses": {
+          "200": {
+            "content": {
+              "*/*": {
+                "schema": {
+                  "properties": {
+                    "data": {
+                      "properties": {
+                        "attributes": {
+                          "properties": {
+                            "api": {
+                              "properties": {
+                                "free_calls": {
+                                  "example": 50,
+                                  "type": "integer"
+                                },
+                                "sizes": {
+                                  "example": "all",
+                                  "type": "string"
+                                }
+                              },
+                              "type": "object"
+                            },
+                            "credits": {
+                              "properties": {
+                                "enterprise": {
+                                  "example": 0,
+                                  "format": "float",
+                                  "type": "number"
+                                },
+                                "payg": {
+                                  "example": 50,
+                                  "format": "float",
+                                  "type": "number"
+                                },
+                                "subscription": {
+                                  "example": 150,
+                                  "format": "float",
+                                  "type": "number"
+                                },
+                                "total": {
+                                  "example": 200,
+                                  "format": "float",
+                                  "type": "number"
+                                }
+                              },
+                              "type": "object"
+                            }
+                          },
+                          "type": "object"
+                        }
+                      },
+                      "type": "object"
+                    }
+                  },
+                  "type": "object"
+                }
+              }
+            },
+            "description": ""
+          },
+          "403": {
+            "content": {
+              "*/*": {
+                "schema": {
+                  "$ref": "#/components/schemas/AuthFailed"
+                }
+              }
+            },
+            "description": "Error: Authentication failed"
+          },
+          "429": {
+            "content": {
+              "*/*": {
+                "schema": {
+                  "$ref": "#/components/schemas/RateLimit"
+                }
+              }
+            },
+            "description": "Error: Rate limit exceeded"
+          }
+        },
+        "summary": "Fetch credit balance and free API calls.",
+        "tags": [
+          "Fetch account info"
+        ]
+      }
+    },
+    "/improve": {
+      "post": {
+        "description": "Submit an image to the remove.bg Improvement program\n* Contribute an image that remove.bg is currently not able to remove the background from properly\n* Help us make remove.bg better\n* Get better results for similiar images in the future\n\nNotes:\n  * By submitting images through the API you agree to the <a target=\"_blank\" rel=\"noopener\" href=\"/ipc\">Improvement Program Conditions</a>\n  * File size: up to 12MB\n  * up to 100 files per day (up to 1000 for Enterprise customers). <br> Higher Rate Limits are available <a href=\"/support/contact?subject=Improvement+Program+Rate+Limit\">upon request</a>.\n\nRequires either an API Key to be provided in the `X-API-Key` request header or an OAuth 2.0 access token to be provided in the `Authorization` request header.\n\nPlease note that submissions are used on a best-effort basis and the extent of expected improvement varies depending on many factors, including the number of provided images, their complexity and visual similarity. Improvements usually take several weeks to become effective.\n",
+        "requestBody": {
+          "content": {
+            "application/json": {
+              "schema": {
+                "$ref": "#/components/schemas/ImprovementProgramJson"
+              }
+            },
+            "application/x-www-form-urlencoded": {
+              "schema": {
+                "$ref": "#/components/schemas/ImprovementProgramJson"
+              }
+            },
+            "multipart/form-data": {
+              "schema": {
+                "$ref": "#/components/schemas/ImprovementProgramMultipart"
+              }
+            }
+          },
+          "required": true
+        },
+        "responses": {
+          "200": {
+            "content": {
+              "*/*": {
+                "schema": {
+                  "$ref": "#/components/schemas/ImprovementProgramJsonResponse"
+                }
+              }
+            },
+            "description": "Image submitted"
+          },
+          "400": {
+            "content": {
+              "*/*": {
+                "schema": {
+                  "properties": {
+                    "errors": {
+                      "items": {
+                        "properties": {
+                          "title": {
+                            "description": "Error message",
+                            "example": "File too large",
+                            "type": "string"
+                          }
+                        },
+                        "type": "object"
+                      },
+                      "type": "array"
+                    }
+                  },
+                  "type": "object"
+                }
+              }
+            },
+            "description": "Error: Invalid parameters or input file unprocessable"
+          },
+          "403": {
+            "content": {
+              "*/*": {
+                "schema": {
+                  "$ref": "#/components/schemas/AuthFailed"
+                }
+              }
+            },
+            "description": "Error: Authentication failed"
+          },
+          "429": {
+            "content": {
+              "*/*": {
+                "schema": {
+                  "$ref": "#/components/schemas/RateLimit"
+                }
+              }
+            },
+            "description": "Error: Rate limit exceeded"
+          }
+        },
+        "tags": [
+          "Improvement Program"
+        ]
+      }
+    },
+    "/removebg": {
+      "post": {
+        "description": "Removes the background of a JPG/PNG image.\n\n* File size: up to 12 MB\n* Image source: File upload (binary or as base64 encoded string) or download from URL\n* Image Content: Any photo with a foreground [(e.g. people, products, animals, cars, etc.)](/supported-images)\n* Output resolutions available: Preview (up to 0.25 megapixels), Full (up to 25 megapixels)\n\nRequires either an API Key to be provided in the `X-API-Key` request header or an OAuth 2.0 access token to be provided in the `Authorization` request header.\n",
+        "requestBody": {
+          "content": {
+            "application/json": {
+              "schema": {
+                "$ref": "#/components/schemas/RemoveBgJson"
+              }
+            },
+            "application/x-www-form-urlencoded": {
+              "schema": {
+                "$ref": "#/components/schemas/RemoveBgJson"
+              }
+            },
+            "multipart/form-data": {
+              "schema": {
+                "$ref": "#/components/schemas/RemoveBgMultipart"
+              }
+            }
+          },
+          "required": true
+        },
+        "responses": {
+          "200": {
+            "content": {
+              "application/json": {
+                "schema": {
+                  "$ref": "#/components/schemas/RemoveBgJsonResponse"
+                }
+              },
+              "image/*": {
+                "schema": {
+                  "example": "(binary image data)",
+                  "format": "binary",
+                  "type": "string"
+                }
+              }
+            },
+            "description": "Image background removed",
+            "headers": {
+              "X-Credits-Charged": {
+                "description": "Amount of credits charged for this call",
+                "schema": {
+                  "format": "float",
+                  "type": "number"
+                }
+              },
+              "X-Height": {
+                "description": "Height of the result image",
+                "schema": {
+                  "type": "integer"
+                }
+              },
+              "X-Type": {
+                "description": "Detected foreground type (How specific this classification is depends on the type_level parameter sent in the request)",
+                "schema": {
+                  "enum": [
+                    "person",
+                    "product",
+                    "animal",
+                    "car",
+                    "car_interior",
+                    "car_part",
+                    "transportation",
+                    "graphics",
+                    "other"
+                  ],
+                  "type": "string"
+                }
+              },
+              "X-Width": {
+                "description": "Width of the result image",
+                "schema": {
+                  "type": "integer"
+                }
+              }
+            }
+          },
+          "400": {
+            "content": {
+              "*/*": {
+                "examples": {
+                  "failed_bg_download": {
+                    "value": {
+                      "errors": [
+                        {
+                          "code": "failed_bg_download",
+                          "title": "Failed to download background image from given image_url: https://unsplash.com/"
+                        }
+                      ]
+                    }
+                  },
+                  "failed_image_download": {
+                    "value": {
+                      "errors": [
+                        {
+                          "code": "failed_image_download",
+                          "title": "Failed to download image from given image_url: https://unsplash.com/"
+                        }
+                      ]
+                    }
+                  },
+                  "failed_to_read_image": {
+                    "value": {
+                      "errors": [
+                        {
+                          "code": "failed_to_read_image",
+                          "title": "There was an error reading the image."
+                        }
+                      ]
+                    }
+                  },
+                  "file_too_large": {
+                    "value": {
+                      "errors": [
+                        {
+                          "code": "file_too_large",
+                          "detail": "File exceeds limit of 12MB",
+                          "title": "File too large"
+                        }
+                      ]
+                    }
+                  },
+                  "invalid_bg_color": {
+                    "value": {
+                      "errors": [
+                        {
+                          "code": "invalid_bg_color",
+                          "title": "Invalid bg_color parameter given"
+                        }
+                      ]
+                    }
+                  },
+                  "invalid_bg_image_url": {
+                    "value": {
+                      "errors": [
+                        {
+                          "code": "invalid_bg_image_url",
+                          "title": "Invalid bg_image_url: Please provide a valid URL."
+                        }
+                      ]
+                    }
+                  },
+                  "invalid_channels": {
+                    "value": {
+                      "errors": [
+                        {
+                          "code": "invalid_channels",
+                          "title": "Invalid value for parameter 'channels'"
+                        }
+                      ]
+                    }
+                  },
+                  "invalid_crop_margin": {
+                    "value": {
+                      "errors": [
+                        {
+                          "code": "invalid_crop_margin",
+                          "title": "Invalid crop_margin parameter given"
+                        }
+                      ]
+                    }
+                  },
+                  "invalid_dimensions": {
+                    "value": {
+                      "errors": [
+                        {
+                          "code": "invalid_dimensions",
+                          "detail": "The dimensions of the given image could not be read.",
+                          "title": "Failed to read image dimensions"
+                        }
+                      ]
+                    }
+                  },
+                  "invalid_file_type": {
+                    "value": {
+                      "errors": [
+                        {
+                          "code": "invalid_file_type",
+                          "detail": "Expected jpg/png, received mp4",
+                          "title": "Invalid file type"
+                        }
+                      ]
+                    }
+                  },
+                  "invalid_format": {
+                    "value": {
+                      "errors": [
+                        {
+                          "code": "invalid_format",
+                          "title": "Invalid format parameter given"
+                        }
+                      ]
+                    }
+                  },
+                  "invalid_image_url": {
+                    "value": {
+                      "errors": [
+                        {
+                          "code": "invalid_image_url",
+                          "title": "Invalid image_url: Please provide a valid URL."
+                        }
+                      ]
+                    }
+                  },
+                  "invalid_position": {
+                    "value": {
+                      "errors": [
+                        {
+                          "code": "invalid_position",
+                          "title": "Invalid position parameter given"
+                        }
+                      ]
+                    }
+                  },
+                  "invalid_roi": {
+                    "value": {
+                      "errors": [
+                        {
+                          "code": "invalid_roi",
+                          "title": "Invalid roi parameter given"
+                        }
+                      ]
+                    }
+                  },
+                  "invalid_scale": {
+                    "value": {
+                      "errors": [
+                        {
+                          "code": "invalid_scale",
+                          "title": "Invalid scale parameter given"
+                        }
+                      ]
+                    }
+                  },
+                  "invalid_semitransparency": {
+                    "value": {
+                      "errors": [
+                        {
+                          "code": "invalid_semitransparency",
+                          "title": "Invalid semitransparency parameter given"
+                        }
+                      ]
+                    }
+                  },
+                  "invalid_size": {
+                    "value": {
+                      "errors": [
+                        {
+                          "code": "invalid_size",
+                          "title": "Invalid value for parameter 'size'"
+                        }
+                      ]
+                    }
+                  },
+                  "missing_source": {
+                    "value": {
+                      "errors": [
+                        {
+                          "code": "missing_source",
+                          "detail": "Please provide the source image in the image_url, image_file or image_file_b64 parameter.",
+                          "title": "No image given"
+                        }
+                      ]
+                    }
+                  },
+                  "multiple_bg_sources": {
+                    "value": {
+                      "errors": [
+                        {
+                          "code": "multiple_bg_sources",
+                          "title": "Multiple background sources given: Please provide either the bg_color, the bg_image_url or the bg_image_file parameter."
+                        }
+                      ]
+                    }
+                  },
+                  "multiple_sources": {
+                    "value": {
+                      "errors": [
+                        {
+                          "code": "multiple_sources",
+                          "title": "Multiple image sources given: Please provide either the image_url, image_file or image_file_b64 parameter."
+                        }
+                      ]
+                    }
+                  },
+                  "resolution_too_high": {
+                    "value": {
+                      "errors": [
+                        {
+                          "code": "resolution_too_high",
+                          "detail": "Input image has 60 megapixels, maximum supported input resolution is 50 megapixels",
+                          "title": "Image resolution too high"
+                        }
+                      ]
+                    }
+                  },
+                  "roi_exceeds_bounds": {
+                    "value": {
+                      "errors": [
+                        {
+                          "code": "roi_exceeds_bounds",
+                          "detail": "The given roi parameter defines a region that exceeds the image bounds",
+                          "title": "ROI exceeds image bounds"
+                        }
+                      ]
+                    }
+                  },
+                  "roi_region_empty": {
+                    "value": {
+                      "errors": [
+                        {
+                          "code": "roi_region_empty",
+                          "detail": "The given roi parameter defines an empty region",
+                          "title": "ROI region is empty"
+                        }
+                      ]
+                    }
+                  },
+                  "unknown_foreground": {
+                    "value": {
+                      "errors": [
+                        {
+                          "code": "unknown_foreground",
+                          "title": "Could not identify foreground in image. For details and recommendations see https://www.remove.bg/supported-images."
+                        }
+                      ]
+                    }
+                  }
+                },
+                "schema": {
+                  "properties": {
+                    "errors": {
+                      "items": {
+                        "properties": {
+                          "code": {
+                            "description": "Error code",
+                            "example": "file_too_large",
+                            "type": "string"
+                          },
+                          "detail": {
+                            "description": "Detailed description of the error",
+                            "example": "File exceeds limit of 12MB",
+                            "type": "string"
+                          },
+                          "title": {
+                            "description": "Error message",
+                            "example": "File too large",
+                            "type": "string"
+                          }
+                        },
+                        "type": "object"
+                      },
+                      "type": "array"
+                    }
+                  },
+                  "type": "object"
+                }
+              }
+            },
+            "description": "Error: Invalid parameters or input file unprocessable (no credits charged)"
+          },
+          "402": {
+            "content": {
+              "*/*": {
+                "schema": {
+                  "properties": {
+                    "errors": {
+                      "items": {
+                        "properties": {
+                          "title": {
+                            "description": "Error message",
+                            "example": "Insufficient credits",
+                            "type": "string"
+                          }
+                        },
+                        "type": "object"
+                      },
+                      "type": "array"
+                    }
+                  },
+                  "type": "object"
+                }
+              }
+            },
+            "description": "Error: Insufficient credits (no credits charged)"
+          },
+          "403": {
+            "content": {
+              "*/*": {
+                "schema": {
+                  "$ref": "#/components/schemas/AuthFailed"
+                }
+              }
+            },
+            "description": "Error: Authentication failed (no credits charged)"
+          },
+          "429": {
+            "content": {
+              "*/*": {
+                "schema": {
+                  "$ref": "#/components/schemas/RateLimit"
+                }
+              }
+            },
+            "description": "Error: Rate limit exceeded (no credits charged)"
+          }
+        },
+        "summary": "Remove the background of an image",
+        "tags": [
+          "Background Removal"
+        ]
+      }
+    }
+  },
+  "components": {
+    "schemas": {
+      "AuthFailed": {
+        "properties": {
+          "errors": {
+            "items": {
+              "properties": {
+                "title": {
+                  "description": "Error message",
+                  "example": "Missing API Key",
+                  "type": "string"
+                }
+              },
+              "type": "object"
+            },
+            "type": "array"
+          }
+        },
+        "type": "object"
+      },
+      "ImprovementProgramJson": {
+        "properties": {
+          "image_file_b64": {
+            "description": "Source image file (base64-encoded string). (If this parameter is present, the other image source parameters must be empty.)",
+            "example": "",
+            "type": "string"
+          },
+          "image_filename": {
+            "description": "Filename of the image, if not provided it will be autodetected form the submitted data.",
+            "example": "car.jpg",
+            "type": "string"
+          },
+          "image_url": {
+            "description": "Source image URL. (If this parameter is present, the other image source parameters must be empty.)",
+            "example": "https://www.remove.bg/example-hd.jpg",
+            "type": "string"
+          },
+          "tag": {
+            "description": "Images with the same tag are grouped together.",
+            "example": "batch_1_2020",
+            "type": "string"
+          }
+        },
+        "type": "object"
+      },
+      "ImprovementProgramJsonResponse": {
+        "properties": {
+          "id": {
+            "description": "ID of the submitted image",
+            "example": "fe09820b-be79-4511-8836-679583529882",
+            "type": "string"
+          }
+        },
+        "type": "object"
+      },
+      "ImprovementProgramMultipart": {
+        "properties": {
+          "image_file": {
+            "description": "Source image file (binary). (If this parameter is present, the other image source parameters must be empty.)",
+            "format": "binary",
+            "type": "string"
+          },
+          "image_file_b64": {
+            "description": "Source image file (base64-encoded string). (If this parameter is present, the other image source parameters must be empty.)",
+            "example": "",
+            "type": "string"
+          },
+          "image_filename": {
+            "description": "Filename of the image, if not provided it will be autodetected form the submitted data.",
+            "type": "string"
+          },
+          "image_url": {
+            "description": "Source image URL. (If this parameter is present, the other image source parameters must be empty.)",
+            "example": "https://www.remove.bg/example-hd.jpg",
+            "type": "string"
+          },
+          "tag": {
+            "description": "Images with the same tag are grouped together.",
+            "type": "string"
+          }
+        },
+        "type": "object"
+      },
+      "RateLimit": {
+        "properties": {
+          "errors": {
+            "items": {
+              "properties": {
+                "title": {
+                  "description": "Error message",
+                  "example": "Rate limit exceeded",
+                  "type": "string"
+                }
+              },
+              "type": "object"
+            },
+            "type": "array"
+          }
+        },
+        "type": "object"
+      },
+      "RemoveBgJson": {
+        "properties": {
+          "add_shadow": {
+            "default": "false",
+            "description": "Whether to add an artificial shadow to the result (default: false). NOTE: Adding shadows is currently only supported for car photos. Other subjects are returned without shadow, even if set to true (this might change in the future).\n",
+            "type": "boolean"
+          },
+          "bg_color": {
+            "description": "Adds a solid color background. Can be a hex color code (e.g. 81d4fa, fff) or a color name (e.g. green). For semi-transparency, 4-/8-digit hex codes are also supported (e.g. 81d4fa77). (If this parameter is present, the other bg_ parameters must be empty.)\n",
+            "example": "",
+            "type": "string"
+          },
+          "bg_image_url": {
+            "description": "Adds a background image from a URL. The image is centered and resized to fill the canvas while preserving the aspect ratio, unless it already has the exact same dimensions as the foreground image. (If this parameter is present, the other bg_ parameters must be empty.)",
+            "example": "",
+            "type": "string"
+          },
+          "channels": {
+            "default": "rgba",
+            "description": "Request either the finalized image (\"rgba\", default) or an alpha mask (\"alpha\"). Note: Since remove.bg also applies RGB color corrections on edges, using only the alpha mask often leads to a lower final image quality. Therefore \"rgba\" is recommended.\n",
+            "enum": [
+              "rgba",
+              "alpha"
+            ],
+            "type": "string"
+          },
+          "crop": {
+            "default": "false",
+            "description": "Whether to crop off all empty regions (default: false). Note that cropping has no effect on the amount of charged credits.\n",
+            "type": "boolean"
+          },
+          "crop_margin": {
+            "default": "0",
+            "description": "Adds a margin around the cropped subject (default: 0). Can be an absolute value (e.g. \"30px\") or relative to the subject size (e.g. \"10%\"). Can be a single value (all sides), two values (top/bottom and left/right) or four values (top, right, bottom, left). This parameter only has an effect when \"crop=true\". The maximum margin that can be added on each side is 50% of the subject dimensions or 500 pixels.\n",
+            "type": "string"
+          },
+          "format": {
+            "default": "auto",
+            "description": "Result image format:\n\"auto\" = Use PNG format if transparent regions exist, otherwise use JPG format (default),\n\"png\" = PNG format with alpha transparency,\n\"jpg\" = JPG format, no transparency,\n\"zip\" = ZIP format, contains color image and alpha matte image, supports transparency (recommended).\n",
+            "enum": [
+              "auto",
+              "png",
+              "jpg",
+              "zip"
+            ],
+            "type": "string"
+          },
+          "image_file_b64": {
+            "description": "Source image file (base64-encoded string). (If this parameter is present, the other image source parameters must be empty.)",
+            "example": "",
+            "type": "string"
+          },
+          "image_url": {
+            "description": "Source image URL. (If this parameter is present, the other image source parameters must be empty.)",
+            "example": "https://www.remove.bg/example-hd.jpg",
+            "type": "string"
+          },
+          "position": {
+            "default": "original",
+            "description": "Positions the subject within the image canvas. Can be \"original\" (default unless \"scale\" is given), \"center\" (default when \"scale\" is given) or a value from \"0%\" to \"100%\" (both horizontal and vertical) or two values (horizontal, vertical).\n",
+            "type": "string"
+          },
+          "roi": {
+            "default": "0% 0% 100% 100%",
+            "description": "Region of interest: Only contents of this rectangular region can be detected as foreground. Everything outside is considered background and will be removed. The rectangle is defined as two x/y coordinates in the format \"x1 y1 x2 y2\". The coordinates can be in absolute pixels (suffix 'px') or relative to the width/height of the image (suffix '%'). By default, the whole image is the region of interest (\"0% 0% 100% 100%\").\n",
+            "type": "string"
+          },
+          "scale": {
+            "default": "original",
+            "description": "Scales the subject relative to the total image size. Can be any value from \"10%\" to \"100%\", or \"original\" (default). Scaling the subject implies \"position=center\" (unless specified otherwise).\n",
+            "type": "string"
+          },
+          "semitransparency": {
+            "default": "true",
+            "description": "Whether to have semi-transparent regions in the result (default: true). NOTE: Semitransparency is currently only supported for car windows (this might change in the future). Other objects are returned without semitransparency, even if set to true.\n",
+            "type": "boolean"
+          },
+          "size": {
+            "default": "preview",
+            "description": "Maximum output image resolution:\n\"preview\" (default) = Resize image to 0.25 megapixels (e.g. 625×400 pixels) – 0.25 credits per image,\n\"full\" = Use original image resolution, up to 25 megapixels (e.g. 6250x4000) with formats ZIP or JPG, or up to 10 megapixels (e.g. 4000x2500) with PNG – 1 credit per image),\n\"auto\" = Use highest available resolution (based on image size and available credits).\n\nFor backwards-compatibility this parameter also accepts the values \"medium\" (up to 1.5 megapixels) and \"hd\" (up to 4 megapixels) for 1 credit per image. The value \"full\" is also available under the name \"4k\" and the value \"preview\" is aliased as \"small\" and \"regular\".\n",
+            "enum": [
+              "preview",
+              "full",
+              "auto"
+            ],
+            "type": "string"
+          },
+          "type": {
+            "default": "auto",
+            "description": "Foreground type:\n\"auto\" = Automatically detect kind of foreground,\n\"person\" = Use person(s) as foreground,\n\"product\" = Use product(s) as foreground.\n\"car\" = Use car as foreground,\n",
+            "enum": [
+              "auto",
+              "person",
+              "product",
+              "car"
+            ],
+            "type": "string"
+          },
+          "type_level": {
+            "default": 1,
+            "description": "Classification level of the detected foreground type:\n\"none\" = No classification (X-Type Header won't bet set on the response)\n\"1\" = Use coarse classification classes: [person, product, animal, car, other]\n\"2\" = Use more specific classification classes: [person, product, animal, car, car_interior, car_part, transportation, graphics, other]\n\"latest\" = Always use the latest classification classes available\n",
+            "enum": [
+              "none",
+              1,
+              2,
+              "latest"
+            ],
+            "type": "string"
+          }
+        },
+        "type": "object"
+      },
+      "RemoveBgJsonResponse": {
+        "properties": {
+          "data": {
+            "properties": {
+              "result_b64": {
+                "description": "Base64 encoded string of result image",
+                "example": "iVBORw0KGgoAAAANSUhEUgAAAIsAAACFC...",
+                "type": "string"
+              }
+            },
+            "type": "object"
+          }
+        },
+        "type": "object"
+      },
+      "RemoveBgMultipart": {
+        "properties": {
+          "add_shadow": {
+            "default": "false",
+            "description": "Whether to add an artificial shadow to the result (default: false). NOTE: Adding shadows is currently only supported for car photos. Other subjects are returned without shadow, even if set to true (this might change in the future).\n",
+            "type": "boolean"
+          },
+          "bg_color": {
+            "description": "Adds a solid color background. Can be a hex color code (e.g. 81d4fa, fff) or a color name (e.g. green). For semi-transparency, 4-/8-digit hex codes are also supported (e.g. 81d4fa77). (If this parameter is present, the other bg_ parameters must be empty.)\n",
+            "example": "",
+            "type": "string"
+          },
+          "bg_image_file": {
+            "description": "Adds a background image from a file (binary). The image is centered and resized to fill the canvas while preserving the aspect ratio, unless it already has the exact same dimensions as the foreground image. (If this parameter is present, the other bg_ parameters must be empty.)",
+            "format": "binary",
+            "type": "string"
+          },
+          "bg_image_url": {
+            "description": "Adds a background image from a URL. The image is centered and resized to fill the canvas while preserving the aspect ratio, unless it already has the exact same dimensions as the foreground image. (If this parameter is present, the other bg_ parameters must be empty.)",
+            "example": "",
+            "type": "string"
+          },
+          "channels": {
+            "default": "rgba",
+            "description": "Request either the finalized image (\"rgba\", default) or an alpha mask (\"alpha\"). Note: Since remove.bg also applies RGB color corrections on edges, using only the alpha mask often leads to a lower final image quality. Therefore \"rgba\" is recommended.\n",
+            "enum": [
+              "rgba",
+              "alpha"
+            ],
+            "type": "string"
+          },
+          "crop": {
+            "default": "false",
+            "description": "Whether to crop off all empty regions (default: false). Note that cropping has no effect on the amount of charged credits.\n",
+            "type": "boolean"
+          },
+          "crop_margin": {
+            "default": "0",
+            "description": "Adds a margin around the cropped subject (default: 0). Can be an absolute value (e.g. \"30px\") or relative to the subject size (e.g. \"10%\"). Can be a single value (all sides), two values (top/bottom and left/right) or four values (top, right, bottom, left). This parameter only has an effect when \"crop=true\". The maximum margin that can be added on each side is 50% of the subject dimensions or 500 pixels.\n",
+            "type": "string"
+          },
+          "format": {
+            "default": "auto",
+            "description": "Result image format:\n\"auto\" = Use PNG format if transparent regions exist, otherwise use JPG format (default),\n\"png\" = PNG format with alpha transparency,\n\"jpg\" = JPG format, no transparency,\n\"zip\" = ZIP format, contains color image and alpha matte image, supports transparency (recommended).\n",
+            "enum": [
+              "auto",
+              "png",
+              "jpg",
+              "zip"
+            ],
+            "type": "string"
+          },
+          "image_file": {
+            "description": "Source image file (binary). (If this parameter is present, the other image source parameters must be empty.)",
+            "format": "binary",
+            "type": "string"
+          },
+          "image_file_b64": {
+            "description": "Source image file (base64-encoded string). (If this parameter is present, the other image source parameters must be empty.)",
+            "example": "",
+            "type": "string"
+          },
+          "image_url": {
+            "description": "Source image URL. (If this parameter is present, the other image source parameters must be empty.)",
+            "example": "https://www.remove.bg/example-hd.jpg",
+            "type": "string"
+          },
+          "position": {
+            "default": "original",
+            "description": "Positions the subject within the image canvas. Can be \"original\" (default unless \"scale\" is given), \"center\" (default when \"scale\" is given) or a value from \"0%\" to \"100%\" (both horizontal and vertical) or two values (horizontal, vertical).\n",
+            "type": "string"
+          },
+          "roi": {
+            "default": "0% 0% 100% 100%",
+            "description": "Region of interest: Only contents of this rectangular region can be detected as foreground. Everything outside is considered background and will be removed. The rectangle is defined as two x/y coordinates in the format \"x1 y1 x2 y2\". The coordinates can be in absolute pixels (suffix 'px') or relative to the width/height of the image (suffix '%'). By default, the whole image is the region of interest (\"0% 0% 100% 100%\").\n",
+            "type": "string"
+          },
+          "scale": {
+            "default": "original",
+            "description": "Scales the subject relative to the total image size. Can be any value from \"10%\" to \"100%\", or \"original\" (default). Scaling the subject implies \"position=center\" (unless specified otherwise).\n",
+            "type": "string"
+          },
+          "semitransparency": {
+            "default": "true",
+            "description": "Whether to have semi-transparent regions in the result (default: true). NOTE: Semitransparency is currently only supported for car windows (this might change in the future). Other objects are returned without semitransparency, even if set to true.\n",
+            "type": "boolean"
+          },
+          "size": {
+            "default": "preview",
+            "description": "Maximum output image resolution:\n\"preview\" (default) = Resize image to 0.25 megapixels (e.g. 625×400 pixels) – 0.25 credits per image,\n\"full\" = Use original image resolution, up to 25 megapixels (e.g. 6250x4000) with formats ZIP or JPG, or up to 10 megapixels (e.g. 4000x2500) with PNG – 1 credit per image),\n\"auto\" = Use highest available resolution (based on image size and available credits).\n\nFor backwards-compatibility this parameter also accepts the values \"medium\" (up to 1.5 megapixels) and \"hd\" (up to 4 megapixels) for 1 credit per image. The value \"full\" is also available under the name \"4k\" and the value \"preview\" is aliased as \"small\" and \"regular\".\n",
+            "enum": [
+              "preview",
+              "full",
+              "auto"
+            ],
+            "type": "string"
+          },
+          "type": {
+            "default": "auto",
+            "description": "Foreground type:\n\"auto\" = Automatically detect kind of foreground,\n\"person\" = Use person(s) as foreground,\n\"product\" = Use product(s) as foreground.\n\"car\" = Use car as foreground,\n",
+            "enum": [
+              "auto",
+              "person",
+              "product",
+              "car"
+            ],
+            "type": "string"
+          },
+          "type_level": {
+            "default": 1,
+            "description": "Classification level of the detected foreground type:\n\"none\" = No classification (X-Type Header won't bet set on the response)\n\"1\" = Use coarse classification classes: [person, product, animal, car, other]\n\"2\" = Use more specific classification classes: [person, product, animal, car, car_interior, car_part, transportation, graphics, other]\n\"latest\" = Always use the latest classification classes available\n",
+            "enum": [
+              "none",
+              1,
+              2,
+              "latest"
+            ],
+            "type": "string"
+          }
+        },
+        "type": "object"
+      }
+    },
+    "securitySchemes": {
+      "APIKeyHeader": {
+        "in": "header",
+        "name": "X-API-Key",
+        "type": "apiKey"
+      }
+    }
+  }
+}
\ No newline at end of file
