WebAPI re-submit

diff --git a/Engine/Plugins/Web/WebAPI/Source/WebAPI/Private/Security/WebAPIAuthentication.cpp b/Engine/Plugins/Web/WebAPI/Source/WebAPI/Private/Security/WebAPIAuthentication.cpp
new file mode 100644
--- /dev/null
+++ b/Engine/Plugins/Web/WebAPI/Source/WebAPI/Private/Security/WebAPIAuthentication.cpp
@@ -0,0 +1,163 @@
+﻿// Copyright Epic Games, Inc. All Rights Reserved.
+
+#include "Security/WebAPIAuthentication.h"
+
+#include "WebAPIDeveloperSettings.h"
+#include "WebAPILog.h"
+#include "WebAPISubsystem.h"
+#include "Engine/Engine.h"
+#include "Interfaces/IHttpRequest.h"
+#include "Serialization/JsonSerializer.h"
+
+UWebAPIOAuthSettings::UWebAPIOAuthSettings()
+{
+	SchemeName = TEXT("OAuth");
+}
+
+bool UWebAPIOAuthSettings::IsValid() const
+{
+	return !ClientId.IsEmpty() && !ClientSecret.IsEmpty();
+}
+
+bool FWebAPIOAuthSchemeHandler::HandleHttpRequest(TSharedPtr<IHttpRequest> InRequest, UWebAPIDeveloperSettings* InSettings)
+{
+	check(InSettings);
+
+	// @todo: cache!
+	UWebAPIOAuthSettings* OAuthSettings = nullptr;
+	InSettings->AuthenticationSettings.FindItemByClass<UWebAPIOAuthSettings>(&OAuthSettings);	
+	check(OAuthSettings);
+
+	if(!ensureMsgf(OAuthSettings->IsValid(), TEXT("Authentication settings are missing one or more required properties.")))
+	{
+		return false;
+	}
+
+	if(OAuthSettings->AccessToken.IsEmpty())
+	{
+		return false;
+	}
+
+	InRequest->SetHeader(TEXT("Authorization"), FString::Printf(TEXT("%s %s"), *OAuthSettings->TokenType, *OAuthSettings->AccessToken));
+
+	return true;
+}
+
+bool FWebAPIOAuthSchemeHandler::HandleHttpResponse(EHttpResponseCodes::Type InResponseCode, TSharedPtr<IHttpResponse, ESPMode::ThreadSafe> InResponse, bool bInWasSuccessful, UWebAPIDeveloperSettings* InSettings)
+{
+	check(InSettings);
+
+	// @todo: cache!
+	UWebAPIOAuthSettings* OAuthSettings = nullptr;
+	InSettings->AuthenticationSettings.FindItemByClass<UWebAPIOAuthSettings>(&OAuthSettings);	
+	check(OAuthSettings);
+	
+	if(!ensureMsgf(OAuthSettings->IsValid(), TEXT("Authentication settings are missing one or more required properties.")))
+	{
+		return false;
+	}
+
+	GEngine->GetEngineSubsystem<UWebAPISubsystem>()->MakeHttpRequest(TEXT("POST"), [&, OAuthSettings](const TSharedPtr<IHttpRequest>& InRequest)
+	{
+		//Microsoft login rest URL
+		FStringFormatNamedArguments UrlArgs;
+		UrlArgs.Add(TEXT("TenantId"), OAuthSettings->TenantId);
+		const FString Url = FString::Format(*OAuthSettings->AuthenticationServer, UrlArgs);
+
+		// @todo: move!
+		FString StringPayload = "grant_type=client_credentials";
+		StringPayload.Append("&client_id=" + OAuthSettings->ClientId);
+		StringPayload.Append("&client_secret=" + OAuthSettings->ClientSecret);
+		StringPayload.Append("&resource=https://digitaltwins.azure.net");
+
+		InRequest->SetURL(Url);
+		InRequest->SetContentAsString(StringPayload);
+
+		TMap<FString, FString> Headers;
+		// FInternetAddr
+		FString SchemeName;
+		FParse::SchemeNameFromURI(*Url, SchemeName);
+		FString Host = Url.Replace(*(SchemeName + TEXT("://")), TEXT(""));
+		
+		int32 DelimiterIdx = -1;
+		if(Host.FindChar(TEXT('/'), DelimiterIdx))
+		{
+			Host = Host.Left(DelimiterIdx);
+		}
+
+		Headers.Add("Host", Host);
+		Headers.Add("Content-Type", "application/x-www-form-urlencoded");
+
+		for (const TPair<FString, FString>& Header : Headers)
+		{
+			InRequest->SetHeader(Header.Key, Header.Value);
+		}
+	})
+	.Next([this, OAuthSettings](const TTuple<FHttpResponsePtr, bool>& InResponse)
+	{
+		// Request failed
+		if(!InResponse.Get<bool>())
+		{
+			return false;
+		}
+
+		const FHttpResponsePtr Response = InResponse.Get<FHttpResponsePtr>();
+
+		FString Message;
+		FString AccessTkn;
+		bool bSuccess = false;
+
+		if (Response.IsValid())
+		{
+			const int32 ResponseCode = Response->GetResponseCode();
+			Message = Response->GetContentAsString();
+
+			if (EHttpResponseCodes::IsOk(ResponseCode))
+			{
+				const TSharedRef<TJsonReader<>> JsonReader = TJsonReaderFactory<>::Create(Response->GetContentAsString());
+				TSharedPtr<FJsonObject> JsonObject;
+				if(FJsonSerializer::Deserialize(JsonReader, JsonObject))
+				{
+					const int32 UnixTimeExpire = JsonObject->GetNumberField("expires_on");
+
+					UE_LOG(LogWebAPI, Display, TEXT("Generate token Response success"));
+					OAuthSettings->TokenType = JsonObject->GetStringField("token_type");
+					OAuthSettings->AccessToken = JsonObject->GetStringField("access_token");
+					OAuthSettings->ExpiresOn =  FDateTime::FromUnixTimestamp(UnixTimeExpire);
+					OAuthSettings->SaveConfig();
+					bSuccess = true;
+				}
+				else
+				{
+					UE_LOG(LogWebAPI, Warning, TEXT("Deserialize JSON failed"));
+					UE_LOG(LogWebAPI, Error, TEXT("Authentication failed: Deserialize JSON Response token failed"));
+					Message = "Deserialize JSON Response token failed:" + Message;
+				}
+			}
+			else
+				{
+				UE_LOG(LogWebAPI, Error, TEXT("Authentication failed: Response not valid"));
+				Message = "Response code not valid:" + Message;
+			}
+		}
+		else
+		{
+			UE_LOG(LogWebAPI, Error, TEXT("Authentication failed: Generate token Response not valid"));
+			Message = "Response is null";
+		}
+
+		// @todo:
+		//Redo all the requests that arrived while fetching new token
+		/*
+		if (bUpdatingToken && bSuccess)
+		{
+			UpdatingToken = false;
+			RedoRequests();
+		}
+		*/
+
+		return bSuccess;
+	});
+
+	return true;
+}
diff --git a/Engine/Plugins/Web/WebAPI/Source/WebAPIBlueprintGraph/Private/K2Node_WebAPIOperation.cpp b/Engine/Plugins/Web/WebAPI/Source/WebAPIBlueprintGraph/Private/K2Node_WebAPIOperation.cpp
new file mode 100644
--- /dev/null
+++ b/Engine/Plugins/Web/WebAPI/Source/WebAPIBlueprintGraph/Private/K2Node_WebAPIOperation.cpp
@@ -0,0 +1,379 @@
+﻿// Copyright Epic Games, Inc. All Rights Reserved.
+
+#include "K2Node_WebAPIOperation.h"
+
+#include "BlueprintActionDatabaseRegistrar.h"
+#include "BlueprintTypePromotion.h"
+#include "EdGraphSchema_K2.h"
+#include "K2Node_StructOperation.h"
+#include "KismetCompiler.h"
+#include "ScopedTransaction.h"
+#include "ToolMenu.h"
+#include "WebAPIOperationObject.h"
+#include "Algo/AllOf.h"
+#include "EdGraph/EdGraphNode.h"
+#include "EdGraph/EdGraphPin.h"
+#include "Kismet2/BlueprintEditorUtils.h"
+
+#define LOCTEXT_NAMESPACE "K2Node_WebAPIAsyncOperation"
+
+struct FGetPinName
+{
+	static const FName& GetMessagePinName() { static const FName MessagePinName(TEXT("Message")); return MessagePinName; }
+	static const FName& GetSuccessPinName() { static const FName SuccessPinName(TEXT("OnSuccess")); return SuccessPinName; }
+	static const FName& GetErrorPinName() { static const FName ErrorPinName(TEXT("OnError")); return ErrorPinName; }
+};
+
+UK2Node_WebAPIOperation::UK2Node_WebAPIOperation()
+{
+	ProxyActivateFunctionName = NAME_None;
+
+	// @todo:
+	// 1. on new node spawned, break Request (input) pin by default: FBlueprintEditor::OnSplitStructPin()
+	// 2. option to convert to/from callback styles: UK2Node_DynamicCast::SetPurity(bool bNewPurity)
+	// 3. collapse output responses that inherit MessageResponse to single message output pin
+}
+
+void UK2Node_WebAPIOperation::AllocateDefaultPins()
+{
+	Super::AllocateDefaultPins();
+}
+
+FText UK2Node_WebAPIOperation::GetNodeTitle(ENodeTitleType::Type TitleType) const
+{
+	FText FunctionName = Super::GetNodeTitle(TitleType);
+	FText NamespaceName;
+	FText ServiceName;
+	
+	if (const UFunction* Function = GetFactoryFunction())
+	{
+		NamespaceName = Function->GetOuterUClass()->GetMetaDataText(TEXT("Namespace"));
+		ServiceName = Function->GetOuterUClass()->GetMetaDataText(TEXT("Service"));
+	}
+	
+	if(TitleType == ENodeTitleType::FullTitle)
+	{
+		FFormatNamedArguments Args;
+		Args.Add(TEXT("FunctionName"), FunctionName);
+		Args.Add(TEXT("NamespaceName"), NamespaceName);
+		Args.Add(TEXT("ServiceName"), ServiceName);
+
+		if (NamespaceName.IsEmpty() && ServiceName.IsEmpty())
+		{
+			return FText::Format(LOCTEXT("NodeTitle", "{FunctionName}"), Args);
+		}
+		else if (NamespaceName.IsEmpty())
+		{
+			return FText::Format(LOCTEXT("NodeTitle_WithNamespace", "{FunctionName}\n{ServiceName}"), Args);
+		}
+		else if (ServiceName.IsEmpty())
+		{
+			return FText::Format(LOCTEXT("NodeTitle_WithService", "{FunctionName}\n{NamespaceName}"), Args);
+		}
+		else
+		{
+			return FText::Format(LOCTEXT("NodeTitle_WithNamespaceAndService", "{FunctionName}\n{NamespaceName}: {ServiceName}"), Args);
+		}		
+	}
+	else
+	{
+		return FunctionName;
+	}
+}
+
+void UK2Node_WebAPIOperation::GetNodeContextMenuActions(UToolMenu* Menu, UGraphNodeContextMenuContext* Context) const
+{
+	Super::GetNodeContextMenuActions(Menu, Context);
+
+	if (!Context->bIsDebugging)
+	{
+		{
+			FText MenuEntryTitle = LOCTEXT("MakeCallbackTitle", "Convert to Callback responses");
+			FText MenuEntryTooltip = LOCTEXT("MakeCallbackTooltip", "Removes the execution pins and instead adds callbacks.");
+
+			bool bCanToggleAsyncType = true;
+			auto CanExecuteAsyncTypeToggle = [](const bool bInCanToggleAsyncType)->bool
+			{
+				return bInCanToggleAsyncType;
+			};
+
+			if (OperationAsyncType == EWebAPIOperationAsyncType::Callback)
+			{
+				MenuEntryTitle = LOCTEXT("MakeLatentActionTitle", "Convert to Latent Action");
+				MenuEntryTooltip = LOCTEXT("MakeLatentActionTooltip", "Adds in branching execution pins so that you can separatly handle the responses.");
+
+				const UEdGraphSchema_K2* K2Schema = Cast<UEdGraphSchema_K2>(GetSchema());
+				check(K2Schema != nullptr);
+
+				bCanToggleAsyncType = K2Schema->DoesGraphSupportImpureFunctions(GetGraph()); // @todo 
+				if (!bCanToggleAsyncType)
+				{
+					MenuEntryTooltip = LOCTEXT("CannotMakeLatentActionTooltip", "This graph does not support latent actions.");
+				}
+			}
+
+			FToolMenuSection& Section = Menu->AddSection("K2NodeWebAPIOperation", LOCTEXT("AsyncTypeHeader", "Async Type"));
+			Section.AddMenuEntry(
+				"ToggleAsyncType",
+				MenuEntryTitle,
+				MenuEntryTooltip,
+				FSlateIcon(),
+				FUIAction(
+					FExecuteAction::CreateUObject(const_cast<UK2Node_WebAPIOperation*>(this), &UK2Node_WebAPIOperation::ToggleAsyncType),
+					FCanExecuteAction::CreateStatic(CanExecuteAsyncTypeToggle, bCanToggleAsyncType),
+					FIsActionChecked()
+				)
+			);
+		}
+	}
+}
+
+bool UK2Node_WebAPIOperation::CanPasteHere(const UEdGraph* TargetGraph) const
+{
+	// @todo: unlike default async node, you can paste into a function (not event graph) and this should convert to callback style
+	return Super::CanPasteHere(TargetGraph);
+}
+
+void UK2Node_WebAPIOperation::PostPlacedNewNode()
+{
+	Super::PostPlacedNewNode();
+
+	PostReconstructNode();
+}
+
+void UK2Node_WebAPIOperation::PostPasteNode()
+{
+	Super::PostPasteNode();
+
+	// @todo: handle conversion to/from callback style
+}
+
+void UK2Node_WebAPIOperation::GetMenuActions(FBlueprintActionDatabaseRegistrar& ActionRegistrar) const
+{
+	struct GetMenuActions_Utils
+	{
+		static void SetNodeFunc(UEdGraphNode* NewNode, bool /*bIsTemplateNode*/, TWeakObjectPtr<UFunction> FunctionPtr)
+		{
+			UK2Node_WebAPIOperation* AsyncTaskNode = CastChecked<UK2Node_WebAPIOperation>(NewNode);
+			if (FunctionPtr.IsValid())
+			{
+				const UFunction* Func = FunctionPtr.Get();
+				FObjectProperty* ReturnProp = CastFieldChecked<FObjectProperty>(Func->GetReturnProperty());
+
+				AsyncTaskNode->ProxyFactoryFunctionName = Func->GetFName();
+				AsyncTaskNode->ProxyFactoryClass        = Func->GetOuterUClass();
+				AsyncTaskNode->ProxyClass               = ReturnProp->PropertyClass;				
+			}
+		}
+	};
+
+	UClass* NodeClass = GetClass();
+	ActionRegistrar.RegisterClassFactoryActions<UWebAPIOperationObject>(FBlueprintActionDatabaseRegistrar::FMakeFuncSpawnerDelegate::CreateLambda([NodeClass](const UFunction* FactoryFunc)->UBlueprintNodeSpawner*
+	{
+		UBlueprintNodeSpawner* NodeSpawner = UBlueprintFunctionNodeSpawner::Create(FactoryFunc);
+		check(NodeSpawner != nullptr);
+		
+		NodeSpawner->NodeClass = NodeClass;
+
+		const TWeakObjectPtr<UFunction> FunctionPtr = MakeWeakObjectPtr(const_cast<UFunction*>(FactoryFunc));
+		NodeSpawner->CustomizeNodeDelegate = UBlueprintNodeSpawner::FCustomizeNodeDelegate::CreateStatic(GetMenuActions_Utils::SetNodeFunc, FunctionPtr);
+
+		return NodeSpawner;
+	}));
+}
+
+void UK2Node_WebAPIOperation::PostReconstructNode()
+{
+	Super::PostReconstructNode();
+
+	if(IsValid())
+	{
+		SplitRequestPins(GetRequestPins());
+		SplitResponsePins(GetResponsePins());
+		SplitResponsePins(GetErrorResponsePins());
+	}
+}
+
+void UK2Node_WebAPIOperation::SetAsyncType(EWebAPIOperationAsyncType InAsyncType)
+{
+	if (InAsyncType != OperationAsyncType)
+	{
+		OperationAsyncType = InAsyncType;
+
+		const bool bHasBeenConstructed = (Pins.Num() > 0);
+		if (bHasBeenConstructed)
+		{
+			ReconstructNode();
+		}
+	}
+}
+
+bool UK2Node_WebAPIOperation::IsValid() const
+{
+	return ProxyFactoryFunctionName != NAME_None
+		&& ProxyClass != nullptr
+		&& ((bHasCompilerMessage && ErrorType >= EMessageSeverity::Info) || !bHasCompilerMessage);
+}
+
+UEdGraphPin* UK2Node_WebAPIOperation::FindPinContainingName(
+	const EEdGraphPinDirection& InPinDirection,
+	const FString& InName) const
+{
+	return FindPinByPredicate([&InPinDirection, &InName](const UEdGraphPin* InPin)
+	{
+		return InPin->Direction == InPinDirection
+			&& InPin->PinType.PinCategory != UEdGraphSchema_K2::PC_Exec
+			&& InPin->GetName().Contains(InName);
+	});
+}
+
+TArray<UEdGraphPin*> UK2Node_WebAPIOperation::FindPinsContainingName(
+	const EEdGraphPinDirection& InPinDirection,
+	const FString& InName) const
+{
+	return Pins.FilterByPredicate([&InPinDirection, InName](const UEdGraphPin* InPin)
+	{
+		return InPin->Direction == InPinDirection
+			&& InPin->PinType.PinCategory != UEdGraphSchema_K2::PC_Exec
+			&& InPin->GetName().Contains(InName);
+	});
+}
+
+TArray<UEdGraphPin*> UK2Node_WebAPIOperation::GetRequestPins() const
+{
+	TArray<UEdGraphPin*> RequestPins = FindPinsContainingName(EGPD_Input, TEXT("Request"));
+	ensureMsgf(!RequestPins.IsEmpty(), TEXT("The function must contain a parameter with \"Request\" in the name."));	
+	return RequestPins;
+}
+
+TArray<UEdGraphPin*> UK2Node_WebAPIOperation::GetResponsePins() const
+{
+	auto ResponsePins = FindPinsContainingName(EGPD_Output, TEXT("Response"));
+	ensureMsgf(!ResponsePins.IsEmpty(), TEXT("The Operation must contain a delegate with a parameter containing \"Response\" in the name."));	
+	return ResponsePins;
+}
+
+TArray<UEdGraphPin*> UK2Node_WebAPIOperation::GetErrorResponsePins() const
+{	
+	auto ErrorResponsePins = FindPinsContainingName(EGPD_Output, TEXT("Error"));
+	ensureMsgf(!ErrorResponsePins.IsEmpty(), TEXT("The Operation must contain a delegate with a parameter containing \"Error\" in the name."));	
+	return ErrorResponsePins;
+}
+
+// @todo: recursively split where struct is basically a wrapper for another struct
+void UK2Node_WebAPIOperation::SplitRequestPins(const TArray<UEdGraphPin*>& InPins) const
+{
+	// If UWebAPIOperation base class, the pin won't always be valid, so ignore.
+	if(!ensureAlways(!InPins.IsEmpty() && Algo::AllOf(InPins, [](const UEdGraphPin* InPin) { return InPin; })))
+	{
+		return;
+	}
+
+	for(UEdGraphPin* Pin : InPins)
+	{
+		// Split if not already
+		if (Pin->SubPins.Num() == 0)
+		{
+			if(CanSplitPin(Pin))
+			{
+				const UEdGraphSchema_K2* K2Schema = GetDefault<UEdGraphSchema_K2>();
+				K2Schema->SplitPin(Pin);
+			}
+		}
+
+		// Removes the prefix from the auto-named split pins 
+		if (Pin->SubPins.Num() > 0)
+		{
+			for (UEdGraphPin* SubPin : Pin->SubPins)
+			{
+				FString SubPinName = SubPin->PinName.ToString();
+				CleanupPinNameInline(SubPinName);
+
+				//SubPin->Modify();			
+				SubPin->PinFriendlyName = FText::FromString(SubPinName);
+			}
+		}
+	}
+}
+
+// Splits if the Response contains a single property
+void UK2Node_WebAPIOperation::SplitResponsePins(const TArray<UEdGraphPin*>& InPins) const
+{
+	// Compilation errors or stale nodes should be warnings, not errors
+	if(!ensureAlways(!InPins.IsEmpty() && Algo::AllOf(InPins, [](const UEdGraphPin* InPin) { return InPin; })))
+	{
+		return;	
+	}
+
+	for(UEdGraphPin* Pin : InPins)
+	{
+		if (const UScriptStruct* PinStructType = Cast<UScriptStruct>(Pin->PinType.PinSubCategoryObject.Get()))
+		{
+			uint32 PropertyNum = 0;
+			
+			// Split if not already
+			if (Pin->SubPins.Num() == 0 && CanSplitPin(Pin))
+			{
+				for (TFieldIterator<const FProperty> PropertyIt(PinStructType, EFieldIteratorFlags::ExcludeSuper, EFieldIteratorFlags::ExcludeDeprecated); PropertyIt; ++PropertyIt)
+				{
+					if (PropertyIt->HasAnyPropertyFlags(CPF_Transient))
+					{
+						continue;
+					}
+
+					PropertyNum++;
+				}
+
+				if(PropertyNum == 1)
+				{
+					// Split, result should be single pin (cause it only contained a single property!)
+					const UEdGraphSchema_K2* K2Schema = GetDefault<UEdGraphSchema_K2>();
+					K2Schema->SplitPin(Pin);
+				}
+			}
+
+			// Removes the prefix from the auto-named split pins 
+			if (Pin->SubPins.Num() > 0)
+			{
+				for (UEdGraphPin* SubPin : Pin->SubPins)
+				{
+					FString SubPinName = SubPin->PinName.ToString();
+					CleanupPinNameInline(SubPinName);
+
+					//SubPin->Modify();			
+					SubPin->PinFriendlyName = FText::FromString(SubPinName);
+				}
+			}
+		}
+	}
+}
+
+void UK2Node_WebAPIOperation::CleanupPinNameInline(FString& InPinName)
+{
+	int32 UnderscorePosition = -1;
+	if(InPinName.FindLastChar(TCHAR('_'), UnderscorePosition))
+	{
+		InPinName = InPinName.RightChop(UnderscorePosition + 1);
+	}
+}
+
+void UK2Node_WebAPIOperation::ToggleAsyncType()
+{
+	const FText TransactionTitle =
+		OperationAsyncType == EWebAPIOperationAsyncType::Callback
+		? LOCTEXT("ToggleAsyncTypeToLatentAction", "Convert to Latent Action")
+		: LOCTEXT("ToggleAsyncTypeToCallback", "Convert to Callback responses");
+	const FScopedTransaction Transaction(TransactionTitle);
+	Modify();
+
+	SetAsyncType(StaticCast<EWebAPIOperationAsyncType>(FMath::Abs(StaticCast<int8>(OperationAsyncType) - 1)));
+}
+
+bool UK2Node_WebAPIOperation::ReconnectExecPins(TArray<UEdGraphPin*>& OldPins)
+{
+	// @todo: handle when converting to/from callback style
+	
+	return true;
+}
+
+#undef LOCTEXT_NAMESPACE
diff --git a/Engine/Plugins/Web/WebAPI/Source/WebAPIBlueprintGraph/Private/WebAPIBlueprintGraphModule.cpp b/Engine/Plugins/Web/WebAPI/Source/WebAPIBlueprintGraph/Private/WebAPIBlueprintGraphModule.cpp
new file mode 100644
--- /dev/null
+++ b/Engine/Plugins/Web/WebAPI/Source/WebAPIBlueprintGraph/Private/WebAPIBlueprintGraphModule.cpp
@@ -0,0 +1,5 @@
+// Copyright Epic Games, Inc. All Rights Reserved.
+
+#include "Modules/ModuleManager.h"
+
+IMPLEMENT_MODULE(FDefaultModuleImpl, WebAPIBlueprintGraph)
diff --git a/Engine/Plugins/Web/WebAPI/Source/WebAPIEditor/Private/Details/Widgets/SWebAPISchemaParameterRow.h b/Engine/Plugins/Web/WebAPI/Source/WebAPIEditor/Private/Details/Widgets/SWebAPISchemaParameterRow.h
new file mode 100644
--- /dev/null
+++ b/Engine/Plugins/Web/WebAPI/Source/WebAPIEditor/Private/Details/Widgets/SWebAPISchemaParameterRow.h
@@ -0,0 +1,20 @@
+﻿// Copyright Epic Games, Inc. All Rights Reserved.
+
+#pragma once
+
+#include "CoreMinimal.h"
+#include "Details/ViewModels/WebAPIParameterViewModel.h"
+#include "SWebAPISchemaTreeTableRow.h"
+#include "SWebAPISchemaTreeTableRow.inl"
+
+class SWebAPISchemaParameterRow
+	: public SWebAPISchemaTreeTableRow<FWebAPIParameterViewModel>
+{
+public:
+	SLATE_BEGIN_ARGS(SWebAPISchemaParameterRow)
+	{ }
+	SLATE_END_ARGS()
+
+public:
+	void Construct(const FArguments& InArgs, const TSharedRef<FWebAPIParameterViewModel>& InViewModel, const TSharedRef<STableViewBase>& InOwnerTableView);
+};
diff --git a/Engine/Plugins/Web/WebAPI/Source/WebAPIEditor/Private/Dom/WebAPIService.cpp b/Engine/Plugins/Web/WebAPI/Source/WebAPIEditor/Private/Dom/WebAPIService.cpp
new file mode 100644
--- /dev/null
+++ b/Engine/Plugins/Web/WebAPI/Source/WebAPIEditor/Private/Dom/WebAPIService.cpp
@@ -0,0 +1,45 @@
+﻿// Copyright Epic Games, Inc. All Rights Reserved.
+
+#include "Dom/WebAPIService.h"
+
+#include "Dom/WebAPIOperation.h"
+#include "Dom/WebAPIType.h"
+
+#define LOCTEXT_NAMESPACE "WebAPIService"
+
+void UWebAPIService::SetNamespace(const FString& InNamespace)
+{
+	Super::SetNamespace(InNamespace);
+	check(Name.HasTypeInfo());
+
+	if(Name.HasTypeInfo() && !Name.TypeInfo->bIsBuiltinType)
+	{
+		Name.TypeInfo->Namespace = InNamespace;
+	}
+}
+
+void UWebAPIService::Visit(const TFunctionRef<void(IWebAPISchemaObjectInterface*&)> InVisitor)
+{
+	Super::Visit(InVisitor);
+
+	for(const TObjectPtr<UWebAPIOperation>& Operation : Operations)
+	{
+		Operation->Visit(InVisitor);
+	}
+}
+
+#if WITH_EDITOR
+EDataValidationResult UWebAPIService::IsDataValid(TArray<FText>& ValidationErrors)
+{
+	EDataValidationResult ValidationResult = Super::IsDataValid(ValidationErrors);
+
+	for(const TObjectPtr<UWebAPIOperation>& Operation : Operations)
+	{
+		ValidationResult = CombineDataValidationResults(Operation->IsDataValid(ValidationErrors), ValidationResult);
+	}
+
+	return ValidationResult;
+}
+#endif // WITH_EDITOR
+
+#undef LOCTEXT_NAMESPACE
diff --git a/Engine/Plugins/Web/WebAPI/Source/WebAPIOpenAPI/Private/Tests/WebAPISwagger.spec.cpp b/Engine/Plugins/Web/WebAPI/Source/WebAPIOpenAPI/Private/Tests/WebAPISwagger.spec.cpp
new file mode 100644
--- /dev/null
+++ b/Engine/Plugins/Web/WebAPI/Source/WebAPIOpenAPI/Private/Tests/WebAPISwagger.spec.cpp
@@ -0,0 +1,375 @@
+﻿// Copyright Epic Games, Inc. All Rights Reserved.
+
+#if WITH_EDITOR
+#include "WebAPIEditorSettings.h"
+#include "Dom/WebAPIParameter.h"
+#include "Misc/AutomationTest.h"
+#include "Misc/FileHelper.h"
+#include "Misc/Paths.h"
+#include "Serialization/JsonSerializer.h"
+#include "V2/WebAPISwaggerConverter.h"
+#include "V2/WebAPISwaggerConverter.inl"
+#include "V2/WebAPISwaggerFactory.h"
+#include "V2/WebAPISwaggerSchema.h"
+
+#if WITH_DEV_AUTOMATION_TESTS
+
+using namespace UE::WebAPI::OpenAPI;
+
+BEGIN_DEFINE_SPEC(FWebAPISwaggerSpec,
+				TEXT("Plugin.WebAPI.Swagger"),
+				EAutomationTestFlags::EditorContext | EAutomationTestFlags::EngineFilter | EAutomationTestFlags::ApplicationContextMask)
+
+	TSharedPtr<UE::WebAPI::OpenAPI::V2::FSwagger> InputDefinition;
+	TStrongObjectPtr<UWebAPIDefinition> OutputDefinition;
+	TSharedPtr<UE::WebAPI::Swagger::FWebAPISwaggerSchemaConverter> Converter;
+	TStrongObjectPtr<UWebAPIDefinitionFactory> Factory;
+
+	FString GetSampleFile(const FString& InName) const
+	{
+		FString FilePath = FPaths::Combine(FPaths::ProjectPluginsDir(),
+			TEXT("WebAPI"), TEXT("Source"), TEXT("WebAPIOpenAPI"),
+			TEXT("Private"), TEXT("Tests"), TEXT("Samples"), TEXT("V2"), InName + TEXT(".json"));
+		ensure(FPaths::FileExists(FilePath));
+		return FilePath;
+	}
+
+	FString GetAPISample() const
+	{
+		return GetSampleFile(TEXT("petstore_V2.json"));
+	}
+
+	FString GetDigitalTwinsManagementAPISample() const
+	{
+		return GetSampleFile(TEXT("digitaltwins"));
+	}
+
+	FString GetDigitalTwinsAPISample() const
+	{
+		return GetSampleFile(TEXT("digitaltwins_data"));
+	}
+
+	TSharedPtr<FJsonObject> LoadJson(const FString& InFile) const
+	{
+		FString FileContents;
+		FFileHelper::LoadFileToString(FileContents, *InFile);
+	
+		TSharedPtr<FJsonObject> JsonObject;
+		FJsonSerializer::Deserialize(TJsonReaderFactory<TCHAR>::Create(FileContents), JsonObject);
+		return JsonObject;
+	}
+
+END_DEFINE_SPEC(FWebAPISwaggerSpec)
+
+// @todo: remove when done: -ExecCmds="Automation RunTests Plugin.WebAPI.Parser.Swagger2" -testexit="Automation Test Queue Empty" -unattended  -nopause
+void FWebAPISwaggerSpec::Define()
+{
+	BeforeEach([this]
+	{
+		const FString FilePath = GetDigitalTwinsManagementAPISample();
+		const TSharedPtr<FJsonObject> JsonObject = LoadJson(FilePath);
+
+		InputDefinition = MakeShared<UE::WebAPI::OpenAPI::V2::FSwagger>();
+		InputDefinition->FromJson(JsonObject.ToSharedRef());
+					
+		OutputDefinition = TStrongObjectPtr(NewObject<UWebAPIDefinition>());
+
+		Converter = MakeShared<UE::WebAPI::Swagger::FWebAPISwaggerSchemaConverter>(
+			InputDefinition,
+			OutputDefinition->GetWebAPISchema(),
+			OutputDefinition->GetMessageLog().ToSharedRef(),
+			OutputDefinition->GetProviderSettings());
+
+		Factory = TStrongObjectPtr(NewObject<UWebAPISwaggerFactory>());
+	});
+	
+	Describe("PetStore", [this]
+	{
+		It("Parses", [this]
+		{
+			const FString FilePath = GetAPISample();
+			const TSharedPtr<FJsonObject> JsonObject = LoadJson(FilePath);
+
+			UE::WebAPI::OpenAPI::V2::FSwagger Object;
+			Object.FromJson(JsonObject.ToSharedRef());
+
+			TestEqual("Swagger Version", Object.Swagger, "2.0");
+			TestEqual("Info Version", Object.Info->Version, "1.0.5");
+			if(TestTrue("Host Set", Object.Host.IsSet()))
+			{
+				TestEqual("Host", Object.Host.GetValue(), "petstore.swagger.io");
+			}
+
+			if(TestTrue("BasePath Set", Object.BasePath.IsSet()))
+			{
+				TestEqual("BasePath", Object.BasePath.GetValue(), "/v2");	
+			}
+
+			if(TestTrue("Tags Set", Object.Tags.IsSet()))
+			{
+				TestEqual("Tag Num", Object.Tags->Num(), 3);	
+			}
+
+			if(TestTrue("Schemes Set", Object.Schemes.IsSet()))
+			{
+				TestEqual("Scheme Num", Object.Schemes->Num(), 2);
+			}
+
+			TestEqual("Paths Num", Object.Paths.Num(), 14);
+			
+			if(TestTrue("Security Definitions Set", Object.SecurityDefinitions.IsSet()))
+			{
+				TestEqual("Security Definitions Num", Object.SecurityDefinitions->Num(), 2);
+			}
+			
+			if(TestTrue("Definitions Set", Object.Definitions.IsSet()))
+			{
+				TestEqual("Definitions Num", Object.Definitions->Num(), 6);
+			}
+
+			if(TestTrue("External Docs Set", Object.ExternalDocs.IsSet()))
+			{
+				TestEqual("External Docs URL", Object.ExternalDocs.GetValue()->Url, TEXT("http://swagger.io"));
+			}			
+		});
+	});
+
+	Describe("DigitalTwinsManagement", [this]
+	{
+		It("Parses", [this]
+		{
+			const FString FilePath = GetDigitalTwinsManagementAPISample();
+			const TSharedPtr<FJsonObject> JsonObject = LoadJson(FilePath);
+
+			UE::WebAPI::OpenAPI::V2::FSwagger Object;
+			Object.FromJson(JsonObject.ToSharedRef());
+
+			TestEqual("Swagger Version", Object.Swagger, "2.0");
+			TestEqual("Info Version", Object.Info->Version, "2020-12-01");
+			if(TestTrue("Host Set", Object.Host.IsSet()))
+			{
+				TestEqual("Host", Object.Host.GetValue(), "management.azure.com");
+			}
+
+			if(TestTrue("Schemes Set", Object.Schemes.IsSet()))
+			{
+				TestEqual("Scheme Num", Object.Schemes->Num(), 1);
+			}
+
+			TestEqual("Paths Num", Object.Paths.Num(), 11);
+				
+			if(TestTrue("Security Definitions Set", Object.SecurityDefinitions.IsSet()))
+			{
+				TestEqual("Security Definitions Num", Object.SecurityDefinitions->Num(), 1);
+			}
+				
+			if(TestTrue("Definitions Set", Object.Definitions.IsSet()))
+			{
+				TestEqual("Definitions Num", Object.Definitions->Num(), 30);
+			}
+		});
+	});
+
+	Describe("Model Composition", [this]
+	{
+		// @todo: meaningful name
+		It("All-Of", [this]
+		{
+			const TSharedPtr<UE::WebAPI::OpenAPI::V2::FSchema>* InputModel = InputDefinition->Definitions->Find(TEXT("ServiceBus"));
+
+			const TObjectPtr<UWebAPIModel> OutputModel = Converter->ConvertModel(*InputModel);
+			if(TestFalse("Converted Model", OutputModel.IsNull()))
+			{
+				if(TestTrue("All-Of is set", (*InputModel)->AllOf.IsSet()))
+				{
+					for(UE::Json::TJsonReference<UE::WebAPI::OpenAPI::V2::FSchema>& SubType : (*InputModel)->AllOf.GetValue())
+					{
+					
+					}
+				}
+			}
+		});
+	});
+
+	Describe("Type Conversion", [this]
+	{
+		BeforeEach([this]
+		{
+			const FString FilePath = GetDigitalTwinsAPISample();
+			const TSharedPtr<FJsonObject> JsonObject = LoadJson(FilePath);
+
+			InputDefinition = MakeShared<UE::WebAPI::OpenAPI::V2::FSwagger>();
+			InputDefinition->FromJson(JsonObject.ToSharedRef());
+					
+			OutputDefinition = TStrongObjectPtr(NewObject<UWebAPIDefinition>());
+
+			Converter = MakeShared<UE::WebAPI::Swagger::FWebAPISwaggerSchemaConverter>(
+				InputDefinition,
+				OutputDefinition->GetWebAPISchema(),
+				OutputDefinition->GetMessageLog().ToSharedRef(),
+				OutputDefinition->GetProviderSettings());
+
+			Factory = TStrongObjectPtr(NewObject<UWebAPISwaggerFactory>());
+		});
+		// // @todo: meaningful name
+		// It("Test", [this]()
+		// {
+		// 	const TSharedPtr<UE::WebAPI::OpenAPI::V2::FSchema>* DigitalTwinsProperties = InputDefinition->Definitions->Find(TEXT("DigitalTwinsProperties"));
+		//
+		// 	const TObjectPtr<UWebAPITypeInfo> Type = Converter->ResolveType(*DigitalTwinsProperties, TEXT("DigitalTwinsProperties"));
+		//
+		// 	TestTrue("TypeInfo is not null", Type != nullptr);
+		// });
+
+		It("Enum that's non-optional and has a single value should default to that value", [this]
+		{
+			check(Converter->Convert());
+			Factory->PostImportWebAPI(OutputDefinition.Get()).Wait();
+			
+			const FString EnumNameToTest = TEXT("ApiVersion");
+
+			TObjectPtr<UWebAPIModelBase>* FoundModel = OutputDefinition->GetWebAPISchema()->Models.FindByPredicate([&](const TObjectPtr<UWebAPIModelBase>& InModel)
+			{
+				return InModel->IsA(UWebAPIEnum::StaticClass()) && InModel->GetName().Contains(EnumNameToTest);
+			});
+
+			if(!TestNotNull("Found Enum", FoundModel)
+				|| !TestTrue("Found Enum not nullptr", !FoundModel->IsNull()))
+			{
+				return false;				
+			}
+
+			const TObjectPtr<UWebAPIEnum> FoundEnum = Cast<UWebAPIEnum>(*FoundModel);
+			TestEqual("Single value", FoundEnum->Values.Num(), 1);
+			TestTrue("Default value is the first value", FoundEnum->GetDefaultValue().Contains(FoundEnum->Values[0]->Name.ToMemberName(TEXT("EV"))));
+
+ 			return true;
+		});
+
+		It("String Array that's optional should have an empty array as a default value", [this]
+		{
+			check(Converter->Convert());
+			Factory->PostImportWebAPI(OutputDefinition.Get()).Wait();
+			
+			const FString ObjectNameToTest = TEXT("IncomingRelationshipCollection");
+			const FString PropertyNameToTest = TEXT("Value");
+
+			TObjectPtr<UWebAPIModelBase>* FoundModelBase = OutputDefinition->GetWebAPISchema()->Models.FindByPredicate([&](const TObjectPtr<UWebAPIModelBase>& InModel)
+			{
+				return InModel->IsA(UWebAPIModel::StaticClass()) && InModel->GetName().Contains(ObjectNameToTest);
+			});
+
+			if(!TestNotNull("Found Model", FoundModelBase)
+				|| !TestTrue("Found Model not nullptr", !FoundModelBase->IsNull()))
+			{
+				return false;				
+			}
+
+			const TObjectPtr<UWebAPIModel> FoundModel = Cast<UWebAPIModel>(*FoundModelBase);
+
+			const TObjectPtr<UWebAPIProperty>* StringArrayProperty = FoundModel->Properties.FindByPredicate([&](const TObjectPtr<UWebAPIProperty>& InProperty)
+			{
+				return InProperty->Name.ToString(true) == PropertyNameToTest;				
+			});
+
+			const FString DefaultValue = (*StringArrayProperty)->GetDefaultValue();
+			TestTrue("Empty Array default value", DefaultValue.Len() <= 3);
+
+			return true;
+		});
+		
+		It("Default value should be set in metadata", [this]
+		{
+
+		});
+
+		It("Properties marked as required in the root of a definition are required in the output object(s)", [this]
+		{
+			check(Converter->Convert());
+			Factory->PostImportWebAPI(OutputDefinition.Get()).Wait();
+					
+			const FString EnumNameToTest = TEXT("ApiVersion");
+
+			TObjectPtr<UWebAPIModelBase>* FoundModel = OutputDefinition->GetWebAPISchema()->Models.FindByPredicate([&](const TObjectPtr<UWebAPIModelBase>& InModel)
+			{
+				return InModel->GetName().Contains(EnumNameToTest);
+			});
+
+			if(!TestNotNull("Found Enum", FoundModel)
+				|| TestTrue("Found Enum not nullptr", !FoundModel->IsNull()))
+			{
+				return false;				
+			}
+
+			const TObjectPtr<UWebAPIEnum> FoundEnum = Cast<UWebAPIEnum>(*FoundModel);
+			const TObjectPtr<UWebAPIEnumValue>* FoundEnumValue = FoundEnum->Values.FindByPredicate([&](const TObjectPtr<UWebAPIEnumValue>& InEnumValue)
+			{
+				return InEnumValue->Name == OutputDefinition->GetProviderSettings().GetUnsetEnumValueName();
+			});
+
+			TestNull("Unset value not found", FoundEnumValue);
+
+			return true;
+		});
+	});
+
+	Describe("Parameter", [this]
+	{
+		It("Resolves schema-less, built-in type", [this]
+		{
+			const TSharedPtr<UE::WebAPI::OpenAPI::V2::FParameter>* InputParameter = InputDefinition->Parameters->Find(TEXT("resourceName"));
+
+			const TObjectPtr<UWebAPIParameter> OutputParameter = Converter->ConvertParameter(*InputParameter);
+			if(TestFalse("Converted Parameter", OutputParameter.IsNull()))
+			{
+				TestEqual("Type is String", OutputParameter->Type.ToString(true), TEXT("String"));
+				TestTrue("Object is named correctly", OutputParameter->GetFName().ToString().StartsWith(TEXT("WebAPIParameter_ResourceName")));
+			}
+		});
+
+		It("Resolves schema-less, generated enum type", [this]
+		{
+			const TSharedPtr<UE::WebAPI::OpenAPI::V2::FParameter>* InputParameter = InputDefinition->Parameters->Find(TEXT("api-version"));
+
+			const TObjectPtr<UWebAPIParameter> OutputParameter = Converter->ConvertParameter(*InputParameter);
+			if(TestFalse("Converted Parameter", OutputParameter.IsNull()))
+			{
+				TestTrue("Type is Enum", OutputParameter->Type.TypeInfo->IsEnum());
+
+				UWebAPIEnum* Model = Cast<UWebAPIEnum>(OutputParameter->Type.TypeInfo->Model.Get());
+				if(TestTrue("Underlying Enum", Model != nullptr))
+				{
+					// TestTrue("Type is Enum", OutputParameter->Type.TypeInfo->IsEnum());
+				}
+			}
+		});
+	});
+		
+	Describe("Operation", [this]
+	{
+		It("Resolves referenced parameters", [this]
+		{
+			const FString Url = TEXT("/subscriptions/{subscriptionId}/resourceGroups/{resourceGroupName}/providers/Microsoft.DigitalTwins/digitalTwinsInstances/{resourceName}/privateEndpointConnections/{privateEndpointConnectionName}");
+
+			const TSharedPtr<V2::FPath, ESPMode::ThreadSafe>* InputService = InputDefinition->Paths.Find(Url);
+			const FString ServiceTag = TEXT("PrivateEndpoints");
+			const TObjectPtr<UWebAPIService> OutputService = Converter->ConvertService(ServiceTag);
+			
+			const TSharedPtr<UE::WebAPI::OpenAPI::V2::FOperation>& InputOperation = InputDefinition->Paths
+				.Find(Url)
+				->Get()->Get;
+
+			const TObjectPtr<UWebAPIOperation> OutputOperation = Converter->ConvertOperation(Url, TEXT("Get"), InputOperation);
+			const FString OperationName = OutputOperation->Name.ToString(true);
+			
+			if(TestFalse("Converted Operation", OutputOperation.IsNull()))
+			{
+				TestTrue("Object is named correctly", OutputOperation->GetFName().ToString().StartsWith(TEXT("WebAPIOperation_PrivateEndpoints_") + OperationName));
+			}
+		});
+
+	});
+}
+
+#endif
+#endif
