WebAPI non-unity fixes

diff --git a/Engine/Plugins/Web/WebAPI/Source/WebAPIEditor/Private/Details/Widgets/SWebAPISchemaEnumValueRow.h b/Engine/Plugins/Web/WebAPI/Source/WebAPIEditor/Private/Details/Widgets/SWebAPISchemaEnumValueRow.h
--- a/Engine/Plugins/Web/WebAPI/Source/WebAPIEditor/Private/Details/Widgets/SWebAPISchemaEnumValueRow.h
+++ b/Engine/Plugins/Web/WebAPI/Source/WebAPIEditor/Private/Details/Widgets/SWebAPISchemaEnumValueRow.h
@@ -4,6 +4,7 @@
 
 #include "CoreMinimal.h"
 #include "SWebAPISchemaTreeTableRow.h"
+#include "SWebAPISchemaTreeTableRow.inl"
 #include "Details/ViewModels/WebAPIEnumViewModel.h"
 
 class SWebAPISchemaEnumValueRow
diff --git a/Engine/Plugins/Web/WebAPI/Source/WebAPIOpenAPI/Private/V2/WebAPISwaggerConverter.cpp b/Engine/Plugins/Web/WebAPI/Source/WebAPIOpenAPI/Private/V2/WebAPISwaggerConverter.cpp
--- a/Engine/Plugins/Web/WebAPI/Source/WebAPIOpenAPI/Private/V2/WebAPISwaggerConverter.cpp
+++ b/Engine/Plugins/Web/WebAPI/Source/WebAPIOpenAPI/Private/V2/WebAPISwaggerConverter.cpp
@@ -1,6 +1,7 @@
 ﻿// Copyright Epic Games, Inc. All Rights Reserved.
 
 #include "WebAPISwaggerConverter.h"
+#include "WebAPISwaggerConverter.inl"
 
 #include "IWebAPIEditorModule.h"
 #include "WebAPIDefinition.h"
diff --git a/Engine/Plugins/Web/WebAPI/Source/WebAPIOpenAPI/Private/V2/WebAPISwaggerConverter.h b/Engine/Plugins/Web/WebAPI/Source/WebAPIOpenAPI/Private/V2/WebAPISwaggerConverter.h
--- a/Engine/Plugins/Web/WebAPI/Source/WebAPIOpenAPI/Private/V2/WebAPISwaggerConverter.h
+++ b/Engine/Plugins/Web/WebAPI/Source/WebAPIOpenAPI/Private/V2/WebAPISwaggerConverter.h
@@ -13,120 +13,114 @@
 #include "UObject/StrongObjectPtr.h"
 #include "V2/WebAPISwaggerSchema.h"
 
-namespace UE
+namespace UE::WebAPI::Swagger
 {
-	namespace WebAPI
+	class FWebAPISwaggerSchemaConverter
 	{
-		namespace Swagger
-		{
-			class FWebAPISwaggerSchemaConverter
-			{
-			public:
-				FWebAPISwaggerSchemaConverter(
-					const TSharedPtr<const UE::WebAPI::OpenAPI::V2::FSwagger>& InSwagger,
-					UWebAPISchema* InWebAPISchema,
-					const TSharedRef<FWebAPIMessageLog>& InMessageLog,
-					const FWebAPIProviderSettings& InProviderSettings);
+	public:
+		FWebAPISwaggerSchemaConverter(
+			const TSharedPtr<const UE::WebAPI::OpenAPI::V2::FSwagger>& InSwagger,
+			UWebAPISchema* InWebAPISchema,
+			const TSharedRef<FWebAPIMessageLog>& InMessageLog,
+			const FWebAPIProviderSettings& InProviderSettings);
 
-				bool Convert();
+		bool Convert();
 
-			public:
-				FString NameTransformer(const FWebAPINameVariant& InString) const;
+	public:
+		FString NameTransformer(const FWebAPINameVariant& InString) const;
 
-				TObjectPtr<UWebAPITypeInfo> ResolveType(FString InJsonType, FString InTypeHint, FString InDefinitionName, const TSharedPtr<OpenAPI::V2::FSchema>& InSchema = nullptr) const;
+		TObjectPtr<UWebAPITypeInfo> ResolveType(FString InJsonType, FString InTypeHint, FString InDefinitionName, const TSharedPtr<OpenAPI::V2::FSchema>& InSchema = nullptr) const;
 
-				template <typename SchemaType>
-				TObjectPtr<UWebAPITypeInfo> ResolveType(const TSharedPtr<SchemaType>& InSchema,	const FString& InDefinitionName = {}, FString InJsonType = {});
+		template <typename SchemaType>
+		TObjectPtr<UWebAPITypeInfo> ResolveType(const TSharedPtr<SchemaType>& InSchema,	const FString& InDefinitionName = {}, FString InJsonType = {});
 
-				TObjectPtr<UWebAPITypeInfo> GetTypeForContentType(const FString& InContentType);
+		TObjectPtr<UWebAPITypeInfo> GetTypeForContentType(const FString& InContentType);
 
-				template <typename StorageType>
-				FString GetDefaultJsonTypeForStorage(const StorageType& InStorage);
+		template <typename StorageType>
+		FString GetDefaultJsonTypeForStorage(const StorageType& InStorage);
 
-				template <typename ObjectType>
-				TSharedPtr<ObjectType> ResolveReference(const FString& InDefinitionName);
+		template <typename ObjectType>
+		TSharedPtr<ObjectType> ResolveReference(const FString& InDefinitionName);
 
-				template <>
-				TSharedPtr<OpenAPI::V2::FSchema> ResolveReference(const FString& InDefinitionName);
+		/** Only resolves if necessary. */
+		template <typename ObjectType>
+		TSharedPtr<ObjectType> ResolveReference(const Json::TJsonReference<ObjectType>& InJsonReference, FString& OutDefinitionName, bool bInCheck = true);
 
-				template <>
-				TSharedPtr<OpenAPI::V2::FParameter> ResolveReference(const FString& InDefinitionName);
+		//FWebAPINameVariant ResolveModelName(const TObjectPtr<UWebAPIProperty>& InProperty, const TOptional<bool>& bInIsArray, const FWebAPITypeNameVariant& InPossibleName = {});
 
-				/** Only resolves if necessary. */
-				template <typename ObjectType>
-				TSharedPtr<ObjectType> ResolveReference(const Json::TJsonReference<ObjectType>& InJsonReference, FString& OutDefinitionName, bool bInCheck = true);
+		FWebAPINameVariant ResolvePropertyName(const TObjectPtr<UWebAPIProperty>& InProperty, const FWebAPITypeNameVariant& InPotentialName, const TOptional<bool>& bInIsArray = {});
 
-				//FWebAPINameVariant ResolveModelName(const TObjectPtr<UWebAPIProperty>& InProperty, const TOptional<bool>& bInIsArray, const FWebAPITypeNameVariant& InPossibleName = {});
+		template <typename SchemaType>
+		bool IsArray(const TSharedPtr<SchemaType>& InSchema);
 
-				FWebAPINameVariant ResolvePropertyName(const TObjectPtr<UWebAPIProperty>& InProperty, const FWebAPITypeNameVariant& InPotentialName, const TOptional<bool>& bInIsArray = {});
+		template <typename SchemaType, typename ModelType>
+		static bool ConvertModelBase(const TSharedPtr<SchemaType>& InSchema, const TObjectPtr<ModelType>& OutModel);
 
-				template <typename SchemaType>
-				bool IsArray(const TSharedPtr<SchemaType>& InSchema);
+		template <typename SchemaType>
+		TObjectPtr<UWebAPIEnum> ConvertEnum(const TSharedPtr<SchemaType>& InSrcSchema, const FWebAPITypeNameVariant& InEnumTypeName = {}) const;
 
-				template <typename SchemaType, typename ModelType>
-				static bool ConvertModelBase(const TSharedPtr<SchemaType>& InSchema, const TObjectPtr<ModelType>& OutModel);
+		/** Modifies an existing Property using the supplied source object. */
+		template <typename SchemaType>
+		bool PatchProperty(const TSharedPtr<SchemaType>& InSrcSchema, const FWebAPINameVariant& InPropertyName, const TObjectPtr<UWebAPIProperty>& OutProperty);
 
-				template <typename SchemaType>
-				TObjectPtr<UWebAPIEnum> ConvertEnum(const TSharedPtr<SchemaType>& InSrcSchema, const FWebAPITypeNameVariant& InEnumTypeName = {}) const;
+		/** Modifies an existing Property using the supplied source object. */
+		template <typename SchemaType>
+		bool PatchProperty(const FWebAPITypeNameVariant& InModelName, const FWebAPINameVariant& InPropertyName, const TSharedPtr<SchemaType>& InSchema, const FString& InDefinitionName, const TObjectPtr<UWebAPIProperty>& OutProperty);
 
-				/** Modifies an existing Property using the supplied source object. */
-				template <typename SchemaType>
-				bool PatchProperty(const TSharedPtr<SchemaType>& InSrcSchema, const FWebAPINameVariant& InPropertyName, const TObjectPtr<UWebAPIProperty>& OutProperty);
-
-				/** Modifies an existing Property using the supplied source object. */
-				template <typename SchemaType>
-				bool PatchProperty(const FWebAPITypeNameVariant& InModelName, const FWebAPINameVariant& InPropertyName, const TSharedPtr<SchemaType>& InSchema, const FString& InDefinitionName, const TObjectPtr<UWebAPIProperty>& OutProperty);
-
-				/** Modifies an existing Property using the supplied TypeInfo. */
-				bool PatchProperty(const FWebAPITypeNameVariant& InModelName, const FWebAPINameVariant& InPropertyName, const FWebAPITypeNameVariant& InPropertyTypeName, const TObjectPtr<UWebAPIProperty>& OutProperty);
-				
-				/** Creates a new Property from the supplied source object. */
-				template <typename SchemaType>
-				TObjectPtr<UWebAPIProperty> ConvertProperty(const TSharedPtr<SchemaType>& InSrcSchema, const TObjectPtr<UWebAPIModel>& InModel, const FWebAPINameVariant& InPropertyName = {}, const FString& InDefinitionName = {});
+		/** Modifies an existing Property using the supplied TypeInfo. */
+		bool PatchProperty(const FWebAPITypeNameVariant& InModelName, const FWebAPINameVariant& InPropertyName, const FWebAPITypeNameVariant& InPropertyTypeName, const TObjectPtr<UWebAPIProperty>& OutProperty);
+		
+		/** Creates a new Property from the supplied source object. */
+		template <typename SchemaType>
+		TObjectPtr<UWebAPIProperty> ConvertProperty(const TSharedPtr<SchemaType>& InSrcSchema, const TObjectPtr<UWebAPIModel>& InModel, const FWebAPINameVariant& InPropertyName = {}, const FString& InDefinitionName = {});
 /*
-				template <>
-				bool ConvertProperty(const FWebAPITypeNameVariant& InModelName,	const FWebAPINameVariant& InPropertyName, const TSharedPtr<OpenAPI::V2::FParameter>& InSchema, const FString& InDefinitionName, const TObjectPtr<UWebAPIProperty>& OutProperty);
+		template <>
+		bool ConvertProperty(const FWebAPITypeNameVariant& InModelName,	const FWebAPINameVariant& InPropertyName, const TSharedPtr<OpenAPI::V2::FParameter>& InSchema, const FString& InDefinitionName, const TObjectPtr<UWebAPIProperty>& OutProperty);
 */
-				/** Modifies an existing Model using the supplied source object. */
-				template <typename SchemaType>
-				bool PatchModel(const TSharedPtr<SchemaType>& InSrcSchema, const FWebAPITypeNameVariant& InModelTypeName, const TObjectPtr<UWebAPIModel>& OutModel);
+		/** Modifies an existing Model using the supplied source object. */
+		template <typename SchemaType>
+		bool PatchModel(const TSharedPtr<SchemaType>& InSrcSchema, const FWebAPITypeNameVariant& InModelTypeName, const TObjectPtr<UWebAPIModel>& OutModel);
+		
+		/** Creates a new Model from the supplied source object. */
+		template <typename SchemaType>
+		TObjectPtr<UWebAPIModel> ConvertModel(const TSharedPtr<SchemaType>& InSrcSchema, const FWebAPITypeNameVariant& InModelTypeName = {});
+
+		bool ConvertOperationParameter(const FWebAPINameVariant& InParameterName, const TSharedPtr<OpenAPI::V2::FParameter>& InParameter, const FString& InDefinitionName, const TObjectPtr<UWebAPIOperationParameter>& OutParameter);
 
-				/** Modifies an existing Model using the supplied source object. */
-				template <>
-				bool PatchModel(const TSharedPtr<OpenAPI::V2::FParameter>& InSrcSchema, const FWebAPITypeNameVariant& InModelTypeName, const TObjectPtr<UWebAPIModel>& OutModel);
+		TObjectPtr<UWebAPIParameter> ConvertParameter(const TSharedPtr<OpenAPI::V2::FParameter>& InSrcParameter);
+		
+		bool ConvertRequest(const FWebAPITypeNameVariant& InOperationName, const TSharedPtr<OpenAPI::V2::FOperation>& InOperation, const TObjectPtr<UWebAPIOperationRequest>& OutRequest);
 
-				/** Creates a new Model from the supplied source object. */
-				template <typename SchemaType>
-				TObjectPtr<UWebAPIModel> ConvertModel(const TSharedPtr<SchemaType>& InSrcSchema, const FWebAPITypeNameVariant& InModelTypeName = {});
+		bool ConvertResponse(const FWebAPITypeNameVariant& InOperationName, uint32 InResponseCode, const TSharedPtr<OpenAPI::V2::FResponse>& InResponse, const TObjectPtr<UWebAPIOperationResponse>& OutResponse);
 
-				bool ConvertOperationParameter(const FWebAPINameVariant& InParameterName, const TSharedPtr<OpenAPI::V2::FParameter>& InParameter, const FString& InDefinitionName, const TObjectPtr<UWebAPIOperationParameter>& OutParameter);
+		TObjectPtr<UWebAPIOperation> ConvertOperation(const FString& InPath, const FString& InVerb, const TSharedPtr<OpenAPI::V2::FOperation>& InSrcOperation, const FWebAPITypeNameVariant& InOperationTypeName = {});
 
-				TObjectPtr<UWebAPIParameter> ConvertParameter(const TSharedPtr<OpenAPI::V2::FParameter>& InSrcParameter);
-				
-				bool ConvertRequest(const FWebAPITypeNameVariant& InOperationName, const TSharedPtr<OpenAPI::V2::FOperation>& InOperation, const TObjectPtr<UWebAPIOperationRequest>& OutRequest);
+		TObjectPtr<UWebAPIService> ConvertService(const FWebAPINameVariant& InName) const;
+		TObjectPtr<UWebAPIService> ConvertService(const TSharedPtr<OpenAPI::V2::FTag>& InTag) const;
 
-				bool ConvertResponse(const FWebAPITypeNameVariant& InOperationName, uint32 InResponseCode, const TSharedPtr<OpenAPI::V2::FResponse>& InResponse, const TObjectPtr<UWebAPIOperationResponse>& OutResponse);
+		bool ConvertModels(const TMap<FString, TSharedPtr<OpenAPI::V2::FSchema>>& InSchemas, UWebAPISchema* OutSchema);
+		bool ConvertParameters(const TMap<FString, TSharedPtr<OpenAPI::V2::FParameter>>& InParameters, UWebAPISchema* OutSchema);
 
-				TObjectPtr<UWebAPIOperation> ConvertOperation(const FString& InPath, const FString& InVerb, const TSharedPtr<OpenAPI::V2::FOperation>& InSrcOperation, const FWebAPITypeNameVariant& InOperationTypeName = {});
+		bool ConvertSecurity(const UE::WebAPI::OpenAPI::V2::FSwagger& InSwagger, UWebAPISchema* OutSchema);
 
-				TObjectPtr<UWebAPIService> ConvertService(const FWebAPINameVariant& InName) const;
-				TObjectPtr<UWebAPIService> ConvertService(const TSharedPtr<OpenAPI::V2::FTag>& InTag) const;
+		bool ConvertTags(const TArray<TSharedPtr<OpenAPI::V2::FTag>>& InTags, UWebAPISchema* OutSchema) const;
 
-				bool ConvertModels(const TMap<FString, TSharedPtr<OpenAPI::V2::FSchema>>& InSchemas, UWebAPISchema* OutSchema);
-				bool ConvertParameters(const TMap<FString, TSharedPtr<OpenAPI::V2::FParameter>>& InParameters, UWebAPISchema* OutSchema);
+		bool ConvertPaths(const TMap<FString, TSharedPtr<OpenAPI::V2::FPath>>& InPaths, UWebAPISchema* OutSchema);
 
-				bool ConvertSecurity(const UE::WebAPI::OpenAPI::V2::FSwagger& InSwagger, UWebAPISchema* OutSchema);
+	private:
+		TSharedPtr<const UE::WebAPI::OpenAPI::V2::FSwagger> InputSchema;
+		TStrongObjectPtr<UWebAPISchema> OutputSchema;
+		TSharedPtr<FWebAPIMessageLog> MessageLog;
+		FWebAPIProviderSettings ProviderSettings;
+	};
 
-				bool ConvertTags(const TArray<TSharedPtr<OpenAPI::V2::FTag>>& InTags, UWebAPISchema* OutSchema) const;
+	template <>
+	TSharedPtr<OpenAPI::V2::FSchema> FWebAPISwaggerSchemaConverter::ResolveReference(const FString& InDefinitionName);
 
-				bool ConvertPaths(const TMap<FString, TSharedPtr<OpenAPI::V2::FPath>>& InPaths, UWebAPISchema* OutSchema);
+	template <>
+	TSharedPtr<OpenAPI::V2::FParameter> FWebAPISwaggerSchemaConverter::ResolveReference(const FString& InDefinitionName);
 
-			private:
-				TSharedPtr<const UE::WebAPI::OpenAPI::V2::FSwagger> InputSchema;
-				TStrongObjectPtr<UWebAPISchema> OutputSchema;
-				TSharedPtr<FWebAPIMessageLog> MessageLog;
-				FWebAPIProviderSettings ProviderSettings;
-			};
-		}
-	}
+	/** Modifies an existing Model using the supplied source object. */
+	template <>
+	bool FWebAPISwaggerSchemaConverter::PatchModel(const TSharedPtr<OpenAPI::V2::FParameter>& InSrcSchema, const FWebAPITypeNameVariant& InModelTypeName, const TObjectPtr<UWebAPIModel>& OutModel);
 }
diff --git a/Engine/Plugins/Web/WebAPI/Source/WebAPIOpenAPI/Private/V2/WebAPISwaggerConverter.inl b/Engine/Plugins/Web/WebAPI/Source/WebAPIOpenAPI/Private/V2/WebAPISwaggerConverter.inl
--- a/Engine/Plugins/Web/WebAPI/Source/WebAPIOpenAPI/Private/V2/WebAPISwaggerConverter.inl
+++ b/Engine/Plugins/Web/WebAPI/Source/WebAPIOpenAPI/Private/V2/WebAPISwaggerConverter.inl
@@ -3,6 +3,8 @@
 #pragma once
 
 #include "WebAPISwaggerConverter.h"
+
+#include "IWebAPIEditorModule.h"
 #include "V2/WebAPISwaggerProvider.h"
 
 #define LOCTEXT_NAMESPACE "WebAPISwaggerConverter"
@@ -221,7 +223,7 @@ namespace UE::WebAPI::Swagger
 			}
 
 			template <>
-			bool FWebAPISwaggerSchemaConverter::PatchModel<OpenAPI::V2::FParameter>(const TSharedPtr<OpenAPI::V2::FParameter>& InSrcSchema, const FWebAPITypeNameVariant& InModelTypeName, const TObjectPtr<UWebAPIModel>& OutModel)
+			inline bool FWebAPISwaggerSchemaConverter::PatchModel<OpenAPI::V2::FParameter>(const TSharedPtr<OpenAPI::V2::FParameter>& InSrcSchema, const FWebAPITypeNameVariant& InModelTypeName, const TObjectPtr<UWebAPIModel>& OutModel)
 			{
 				FWebAPITypeNameVariant ModelTypeName;
 				if(InModelTypeName.IsValid() && InModelTypeName.HasTypeInfo())
diff --git a/Engine/Plugins/Web/WebAPI/Source/WebAPIOpenAPI/Private/V3/WebAPIOpenAPIConverter.cpp b/Engine/Plugins/Web/WebAPI/Source/WebAPIOpenAPI/Private/V3/WebAPIOpenAPIConverter.cpp
--- a/Engine/Plugins/Web/WebAPI/Source/WebAPIOpenAPI/Private/V3/WebAPIOpenAPIConverter.cpp
+++ b/Engine/Plugins/Web/WebAPI/Source/WebAPIOpenAPI/Private/V3/WebAPIOpenAPIConverter.cpp
@@ -1,6 +1,7 @@
 ﻿// Copyright Epic Games, Inc. All Rights Reserved.
 
 #include "WebAPIOpenAPIConverter.h"
+#include "WebAPIOpenAPIConverter.inl"
 
 #include "IWebAPIEditorModule.h"
 #include "WebAPIDefinition.h"
@@ -35,1868 +36,1190 @@ if(SrcProperty.IsSet())											\
 	DstProperty = SrcProperty.GetValue();						\
 }
 
-namespace UE
+namespace UE::WebAPI::OpenAPI
 {
-	namespace WebAPI
+	FWebAPIOpenAPISchemaConverter::FWebAPIOpenAPISchemaConverter(
+		const TSharedPtr<const UE::WebAPI::OpenAPI::V3::FOpenAPIObject>& InOpenAPI, UWebAPISchema* InWebAPISchema,
+		const TSharedRef<FWebAPIMessageLog>& InMessageLog, const FWebAPIProviderSettings& InProviderSettings):
+		InputSchema(InOpenAPI)
+		, OutputSchema(InWebAPISchema)
+		, MessageLog(InMessageLog)
+		, ProviderSettings(InProviderSettings)
 	{
-		// @todo: remove
-		// @see: https://blogs.sap.com/2018/01/05/open-api-spec-2.0-vs-3.0/
-		// [REDACTED_PROFILE_URL]/openapi-codegen/blob/fcbbee8ce4304400c7d2c605fb0388d6c6e5b64a/lib/orange/downconvert.js
-		namespace OpenAPI
-		{
-			FWebAPIOpenAPISchemaConverter::FWebAPIOpenAPISchemaConverter(
-				const TSharedPtr<const UE::WebAPI::OpenAPI::V3::FOpenAPIObject>& InOpenAPI, UWebAPISchema* InWebAPISchema,
-				const TSharedRef<FWebAPIMessageLog>& InMessageLog, const FWebAPIProviderSettings& InProviderSettings):
-				InputSchema(InOpenAPI)
-				, OutputSchema(InWebAPISchema)
-				, MessageLog(InMessageLog)
-				, ProviderSettings(InProviderSettings)
-			{
-			}
-
-			bool FWebAPIOpenAPISchemaConverter::Convert()
-			{
-				if(InputSchema->Info.IsValid())
-				{
-					OutputSchema->APIName = InputSchema->Info->Title;
-					OutputSchema->Version = InputSchema->Info->Version;
-				}
-
-				// V3 has multiple servers vs singular in V2
-				if(InputSchema->Servers.IsSet()
-					&& !InputSchema->Servers->IsEmpty()
-					&& InputSchema->Servers.GetValue()[0].IsValid())
-				{
-					const TSharedPtr<V3::FServerObject> Server = InputSchema->Servers.GetValue()[0];
-					
-					FString Url = Server->Url;
-					FString Scheme;
-					FParse::SchemeNameFromURI(*Url, Scheme);
-
-					// If Url isn't complete (relative, etc.), it won't have a scheme.
-					if(!Scheme.IsEmpty())
-					{
-						OutputSchema->URISchemes.Add(Scheme);
-					}
-
-					Url = Url.Replace(*(Scheme + TEXT("://")), TEXT(""));					
-					Url.Split(TEXT("/"), &OutputSchema->Host, &OutputSchema->BaseUrl);
-
-					// If Url isn't complete, there may not be a valid host
-					if(OutputSchema->Host.IsEmpty())
-					{
-						MessageLog->LogWarning(LOCTEXT("NoHostProvided", "The specification did not contain a host Url, this should be specified manually in the generated project settings."), FWebAPIOpenAPIProvider::LogName);
-					}
-				}
+	}
 
-				// If no Url schemes provided, add https as default
-				if(OutputSchema->URISchemes.IsEmpty())
-				{
-					OutputSchema->URISchemes = { TEXT("https") };
-				}
-				
-				// Top level decl of tags optional, so find from paths
-				bool bSuccessfullyConverted = InputSchema->Tags.IsSet()
-											? ConvertTags(InputSchema->Tags.GetValue(), OutputSchema.Get())
-											: true;
-
-				bSuccessfullyConverted &= InputSchema->Components.IsValid()
-										  ? ConvertModels(InputSchema->Components->Schemas, OutputSchema.Get())
-										  : true;
-				//bSuccessfullyConverted &= ConvertParameters(InputSchema->Components.Parameters, OutputSchema.Get());
-				bSuccessfullyConverted &= ConvertPaths(InputSchema->Paths, OutputSchema.Get());
-
-				return bSuccessfullyConverted;
-			}
+	bool FWebAPIOpenAPISchemaConverter::Convert()
+	{
+		if(InputSchema->Info.IsValid())
+		{
+			OutputSchema->APIName = InputSchema->Info->Title;
+			OutputSchema->Version = InputSchema->Info->Version;
+		}
+
+		// V3 has multiple servers vs singular in V2
+		if(InputSchema->Servers.IsSet()
+			&& !InputSchema->Servers->IsEmpty()
+			&& InputSchema->Servers.GetValue()[0].IsValid())
+		{
+			const TSharedPtr<V3::FServerObject> Server = InputSchema->Servers.GetValue()[0];
+			
+			FString Url = Server->Url;
+			FString Scheme;
+			FParse::SchemeNameFromURI(*Url, Scheme);
 
-			FString FWebAPIOpenAPISchemaConverter::NameTransformer(const FWebAPINameVariant& InString) const
+			// If Url isn't complete (relative, etc.), it won't have a scheme.
+			if(!Scheme.IsEmpty())
 			{
-				return ProviderSettings.ToPascalCase(InString);
+				OutputSchema->URISchemes.Add(Scheme);
 			}
 
-			TObjectPtr<UWebAPITypeInfo> FWebAPIOpenAPISchemaConverter::ResolveMappedType(const FString& InType)
-			{
-				const TObjectPtr<UWebAPIStaticTypeRegistry> StaticTypeRegistry = IWebAPIEditorModuleInterface::Get().GetStaticTypeRegistry();
-
-				// [REDACTED_PROFILE_URL]/openapi-generator/blob/5d68bd6a03f0c48e838b4fe3b98b7e30858c0373/modules/openapi-generator/src/main/java/org/openapitools/codegen/languages/CppUE4ClientCodegen.java
-				// OpenAPI type to UE type (Prefix, Name)
-				static TMap<FString, TObjectPtr<UWebAPITypeInfo>> TypeMap =
-				{
-					{TEXT("file"), StaticTypeRegistry->FilePath},
-					{TEXT("any"), StaticTypeRegistry->Object},
-					{TEXT("object"), StaticTypeRegistry->Object},
-
-					{TEXT("array"), StaticTypeRegistry->String},
-					{TEXT("boolean"), StaticTypeRegistry->Boolean},
-					{TEXT("byte"), StaticTypeRegistry->Byte},
-					{TEXT("integer"), StaticTypeRegistry->Int32},
-					{TEXT("int"), StaticTypeRegistry->Int32},
-					{TEXT("int32"), StaticTypeRegistry->Int32},
-					{TEXT("short"), StaticTypeRegistry->Int16},
-					{TEXT("int16"), StaticTypeRegistry->Int16},
-					{TEXT("long"), StaticTypeRegistry->Int64},
-					{TEXT("int64"), StaticTypeRegistry->Int64},
-					{TEXT("float"), StaticTypeRegistry->Float},
-					{TEXT("double"), StaticTypeRegistry->Double},
-
-					{TEXT("number"), StaticTypeRegistry->Int32},
-					{TEXT("char"), StaticTypeRegistry->Char},
-					{TEXT("date"), StaticTypeRegistry->DateTime},
-					{TEXT("date-time"), StaticTypeRegistry->DateTime},
-					{TEXT("password"), StaticTypeRegistry->String},
-					{TEXT("string"), StaticTypeRegistry->String},
-					{TEXT("void"), StaticTypeRegistry->Void},
-					{TEXT("null"), StaticTypeRegistry->Nullptr}
-				};
-
-				if (const TObjectPtr<UWebAPITypeInfo>* FoundTypeInfo = TypeMap.Find(InType))
-				{
-					return *FoundTypeInfo;
-				}
-
-				return nullptr;
-			}
+			Url = Url.Replace(*(Scheme + TEXT("://")), TEXT(""));					
+			Url.Split(TEXT("/"), &OutputSchema->Host, &OutputSchema->BaseUrl);
 
-			TObjectPtr<UWebAPITypeInfo> FWebAPIOpenAPISchemaConverter::ResolveType(
-				FString InType,
-				FString InFormat,
-				FString InDefinitionName,
-				const TSharedPtr<OpenAPI::V3::FSchemaObject>& InSchema)
+			// If Url isn't complete, there may not be a valid host
+			if(OutputSchema->Host.IsEmpty())
 			{
-				const TObjectPtr<UWebAPIStaticTypeRegistry> StaticTypeRegistry = IWebAPIEditorModuleInterface::Get().GetStaticTypeRegistry();
-				
-				static TMap<EJson, FString> EJson_ValueToTypeName = {
-					{EJson::None, TEXT("None")},
-					{EJson::Null, TEXT("nullptr")},
-					{EJson::String, TEXT("FString")},
-					{EJson::Boolean, TEXT("bool")},
-					{EJson::Number, TEXT("int32")},
-					{EJson::Array, TEXT("TArray")},
-					{EJson::Object, TEXT("UObject")},
-				};
-
-				static TMap<EJson, TObjectPtr<UWebAPITypeInfo>> EJson_ValueToTypeInfo = {
-					{EJson::None, StaticTypeRegistry->Void},
-					{EJson::Null, StaticTypeRegistry->Nullptr},
-					{EJson::String, StaticTypeRegistry->String},
-					{EJson::Boolean, StaticTypeRegistry->Boolean},
-					{EJson::Number, StaticTypeRegistry->Int32},
-					{EJson::Array, StaticTypeRegistry->String},
-					{EJson::Object, StaticTypeRegistry->Object},
-				};
-
-				static TMap<FString, TObjectPtr<UWebAPITypeInfo>> JsonTypeValueToTypeInfo = {
-					{TEXT("None"), StaticTypeRegistry->Void},
-					{TEXT("Null"), StaticTypeRegistry->Nullptr},
-					{TEXT("String"), StaticTypeRegistry->String},
-					{TEXT("Boolean"), StaticTypeRegistry->Boolean},
-					{TEXT("Number"), StaticTypeRegistry->Int32},
-					{TEXT("Array"), StaticTypeRegistry->String},
-					{TEXT("Object"), StaticTypeRegistry->Object},
-				};
-
-				TObjectPtr<UWebAPITypeInfo> Result = nullptr;
-
-				// If a definition name is supplied, try to find it first
-				if ((InType == TEXT("Object") || InType == TEXT("Array")) && !InDefinitionName.IsEmpty())
-				{
-					if (const TObjectPtr<UWebAPITypeInfo>* FoundTypeInfo = OutputSchema->TypeRegistry->FindGeneratedType(EWebAPISchemaType::Model, InDefinitionName))
-					{
-						Result = *FoundTypeInfo;
-					}
-				}
-
-				// If not found above
-				if (!Result)
-				{
-					// Try specific types
-					if (const TObjectPtr<UWebAPITypeInfo>& FoundMappedTypeInfo = ResolveMappedType(InFormat.IsEmpty() ? InType : InFormat))
-					{
-						return FoundMappedTypeInfo;
-					}
-					// Fallback to basic types
-					else if (const TObjectPtr<UWebAPITypeInfo>* FoundTypeInfo = JsonTypeValueToTypeInfo.Find(InType))
-					{
-						Result = *FoundTypeInfo;
-					}
-
-					// If it's not a built-in type
-					if (Result != nullptr && !Result->bIsBuiltinType)
-					{
-						// Duplicate it with the provided definition name 
-						if (!InDefinitionName.IsEmpty())
-						{
-							// @todo: what if it's an enum? how do you know at this point if it's an enum or struct?
-							// Allow prefix to be set later depending on this?
-							Result = OutputSchema->TypeRegistry->GetOrMakeGeneratedType(EWebAPISchemaType::Model,
-								InDefinitionName,
-								InDefinitionName,
-								Result);
-							Result->Prefix = TEXT("F");
-							ensure(!Result->Name.IsEmpty());
-						}
-						// otherwise duplicate it as an unnamed/partially resolved type 
-						else
-						{
-							// @todo: prevent!
-							Result = Result->Duplicate(OutputSchema->TypeRegistry);
-						}
-					}
-				}
-
-				return Result;
+				MessageLog->LogWarning(LOCTEXT("NoHostProvided", "The specification did not contain a host Url, this should be specified manually in the generated project settings."), FWebAPIOpenAPIProvider::LogName);
 			}
+		}
 
-			TObjectPtr<UWebAPITypeInfo> FWebAPIOpenAPISchemaConverter::GetTypeForContentType(const FString& InContentType)
-			{
-				const TObjectPtr<UWebAPIStaticTypeRegistry> StaticTypeRegistry = IWebAPIEditorModuleInterface::Get().GetStaticTypeRegistry();
-				
-				static TMap<FString, TObjectPtr<UWebAPITypeInfo>> ContentTypeToTypeInfo = {
-					{TEXT("application/json"), StaticTypeRegistry->JsonObject},
-					{TEXT("application/xml"), StaticTypeRegistry->String},
-					{TEXT("text/plain"), StaticTypeRegistry->String}
-				};
-				
-				if(const TObjectPtr<UWebAPITypeInfo>* FoundTypeInfo = ContentTypeToTypeInfo.Find(InContentType.ToLower()))
-				{
-					return *FoundTypeInfo;
-				}
-
-				// For all other cases, use string
-				return StaticTypeRegistry->String;
-			}
+		// If no Url schemes provided, add https as default
+		if(OutputSchema->URISchemes.IsEmpty())
+		{
+			OutputSchema->URISchemes = { TEXT("https") };
+		}
+		
+		// Top level decl of tags optional, so find from paths
+		bool bSuccessfullyConverted = InputSchema->Tags.IsSet()
+									? ConvertTags(InputSchema->Tags.GetValue(), OutputSchema.Get())
+									: true;
+
+		bSuccessfullyConverted &= InputSchema->Components.IsValid()
+								  ? ConvertModels(InputSchema->Components->Schemas, OutputSchema.Get())
+								  : true;
+		//bSuccessfullyConverted &= ConvertParameters(InputSchema->Components.Parameters, OutputSchema.Get());
+		bSuccessfullyConverted &= ConvertPaths(InputSchema->Paths, OutputSchema.Get());
+
+		return bSuccessfullyConverted;
+	}
 
-			template <>
-			FString FWebAPIOpenAPISchemaConverter::GetDefaultJsonTypeForStorage<EWebAPIParameterStorage>(const EWebAPIParameterStorage& InStorage)
-			{
-				static TMap<EWebAPIParameterStorage, FString> StorageToJsonType = {
-					{ EWebAPIParameterStorage::Body, TEXT("object") },
-					{ EWebAPIParameterStorage::Cookie, TEXT("string") },
-					{ EWebAPIParameterStorage::Header, TEXT("string") },
-					{ EWebAPIParameterStorage::Path, TEXT("string") },
-					{ EWebAPIParameterStorage::Query, TEXT("string") },
-				};
-
-				return StorageToJsonType[InStorage];
-			}
+	FString FWebAPIOpenAPISchemaConverter::NameTransformer(const FWebAPINameVariant& InString) const
+	{
+		return ProviderSettings.ToPascalCase(InString);
+	}
 
-			template <>
-			FString FWebAPIOpenAPISchemaConverter::GetDefaultJsonTypeForStorage<EWebAPIResponseStorage>(const EWebAPIResponseStorage& InStorage)
-			{
-				static TMap<EWebAPIResponseStorage, FString> StorageToJsonType = {
-					{ EWebAPIResponseStorage::Body, TEXT("object") },
-					{ EWebAPIResponseStorage::Header, TEXT("string") },
-				};
+	TObjectPtr<UWebAPITypeInfo> FWebAPIOpenAPISchemaConverter::ResolveMappedType(const FString& InType)
+	{
+		const TObjectPtr<UWebAPIStaticTypeRegistry> StaticTypeRegistry = IWebAPIEditorModuleInterface::Get().GetStaticTypeRegistry();
 
-				return StorageToJsonType[InStorage];
-			}
+		// [REDACTED_PROFILE_URL]/openapi-generator/blob/5d68bd6a03f0c48e838b4fe3b98b7e30858c0373/modules/openapi-generator/src/main/java/org/openapitools/codegen/languages/CppUE4ClientCodegen.java
+		// OpenAPI type to UE type (Prefix, Name)
+		static TMap<FString, TObjectPtr<UWebAPITypeInfo>> TypeMap =
+		{
+			{TEXT("file"), StaticTypeRegistry->FilePath},
+			{TEXT("any"), StaticTypeRegistry->Object},
+			{TEXT("object"), StaticTypeRegistry->Object},
+
+			{TEXT("array"), StaticTypeRegistry->String},
+			{TEXT("boolean"), StaticTypeRegistry->Boolean},
+			{TEXT("byte"), StaticTypeRegistry->Byte},
+			{TEXT("integer"), StaticTypeRegistry->Int32},
+			{TEXT("int"), StaticTypeRegistry->Int32},
+			{TEXT("int32"), StaticTypeRegistry->Int32},
+			{TEXT("short"), StaticTypeRegistry->Int16},
+			{TEXT("int16"), StaticTypeRegistry->Int16},
+			{TEXT("long"), StaticTypeRegistry->Int64},
+			{TEXT("int64"), StaticTypeRegistry->Int64},
+			{TEXT("float"), StaticTypeRegistry->Float},
+			{TEXT("double"), StaticTypeRegistry->Double},
+
+			{TEXT("number"), StaticTypeRegistry->Int32},
+			{TEXT("char"), StaticTypeRegistry->Char},
+			{TEXT("date"), StaticTypeRegistry->DateTime},
+			{TEXT("date-time"), StaticTypeRegistry->DateTime},
+			{TEXT("password"), StaticTypeRegistry->String},
+			{TEXT("string"), StaticTypeRegistry->String},
+			{TEXT("void"), StaticTypeRegistry->Void},
+			{TEXT("null"), StaticTypeRegistry->Nullptr}
+		};
 
-			template <typename SchemaType>
-			TObjectPtr<UWebAPITypeInfo> FWebAPIOpenAPISchemaConverter::ResolveType(const TSharedPtr<SchemaType>& InSchema, const FString& InDefinitionName, FString InJsonType)
-			{
-				static_assert(TIsDerivedFrom<SchemaType, OpenAPI::V3::FSchemaObjectBase>::Value, "Type is not derived from OpenAPI::V3::FSchemaObjectBase.");
+		if (const TObjectPtr<UWebAPITypeInfo>* FoundTypeInfo = TypeMap.Find(InType))
+		{
+			return *FoundTypeInfo;
+		}
 
-				const TObjectPtr<UWebAPIStaticTypeRegistry> StaticTypeRegistry = IWebAPIEditorModuleInterface::Get().GetStaticTypeRegistry();
+		return nullptr;
+	}
 
-				FString DefinitionName = InDefinitionName;
-				TObjectPtr<UWebAPITypeInfo> Result = nullptr;
-				
-				TSharedPtr<UE::WebAPI::OpenAPI::V3::FSchemaObject> ItemSchema = nullptr;
-				if(InSchema.IsValid() && InSchema->Items.IsSet())
-				{
-					ItemSchema = ResolveReference(InSchema->Items.GetValue(), DefinitionName);
-					if(!InSchema->Items->GetPath().IsEmpty())
-					{
-						DefinitionName = InSchema->Items->GetLastPathSegment();						
-					}
-					return ResolveType(ItemSchema, DefinitionName);
-				}
+	TObjectPtr<UWebAPITypeInfo> FWebAPIOpenAPISchemaConverter::ResolveType(
+		FString InType,
+		FString InFormat,
+		FString InDefinitionName,
+		const TSharedPtr<OpenAPI::V3::FSchemaObject>& InSchema)
+	{
+		const TObjectPtr<UWebAPIStaticTypeRegistry> StaticTypeRegistry = IWebAPIEditorModuleInterface::Get().GetStaticTypeRegistry();
+		
+		static TMap<EJson, FString> EJson_ValueToTypeName = {
+			{EJson::None, TEXT("None")},
+			{EJson::Null, TEXT("nullptr")},
+			{EJson::String, TEXT("FString")},
+			{EJson::Boolean, TEXT("bool")},
+			{EJson::Number, TEXT("int32")},
+			{EJson::Array, TEXT("TArray")},
+			{EJson::Object, TEXT("UObject")},
+		};
 
-				if(InSchema.IsValid())
-				{
-					Result = ResolveType(InSchema->Type.Get(InJsonType.IsEmpty() ? TEXT("object") : InJsonType), InSchema->Format.Get(TEXT("")), DefinitionName);
-				}
-				else
-				{
-					Result = ResolveType(InJsonType, TEXT(""), DefinitionName);
-					if(!Result)
-					{
-						Result = StaticTypeRegistry->Object;
-					}
-				}
+		static TMap<EJson, TObjectPtr<UWebAPITypeInfo>> EJson_ValueToTypeInfo = {
+			{EJson::None, StaticTypeRegistry->Void},
+			{EJson::Null, StaticTypeRegistry->Nullptr},
+			{EJson::String, StaticTypeRegistry->String},
+			{EJson::Boolean, StaticTypeRegistry->Boolean},
+			{EJson::Number, StaticTypeRegistry->Int32},
+			{EJson::Array, StaticTypeRegistry->String},
+			{EJson::Object, StaticTypeRegistry->Object},
+		};
 
-				if (InSchema.IsValid() && InSchema->Enum.IsSet() && !InSchema->Enum.GetValue().IsEmpty())
-				{
-					if(const TObjectPtr<UWebAPITypeInfo>* FoundGeneratedType = OutputSchema->TypeRegistry->FindGeneratedType(EWebAPISchemaType::Model, NameTransformer(DefinitionName)))
-					{
-						Result = *FoundGeneratedType;						
-					}
-					else
-					{
-						Result = OutputSchema->TypeRegistry->GetOrMakeGeneratedType(
-							EWebAPISchemaType::Model,
-							NameTransformer(DefinitionName),
-							DefinitionName,
-							StaticTypeRegistry->Enum);
-					}
-				}
-				else if (!Result->IsEnum() && (Result == StaticTypeRegistry->Object || Result->ToString(true).IsEmpty()))
-				{
-					if (!DefinitionName.IsEmpty())
-					{
-						if (const TObjectPtr<UWebAPITypeInfo>* FoundBuiltinType = StaticTypeRegistry->FindBuiltinType(DefinitionName))
-						{
-							Result = *FoundBuiltinType;
-						}
-						else if (const TObjectPtr<UWebAPITypeInfo>* FoundGeneratedModelType = OutputSchema->TypeRegistry->FindGeneratedType(EWebAPISchemaType::Model, DefinitionName))
-						{
-							Result = *FoundGeneratedModelType;
-						}
-						else if (const TObjectPtr<UWebAPITypeInfo>* FoundGeneratedParameterType = OutputSchema->TypeRegistry->FindGeneratedType(EWebAPISchemaType::Parameter, DefinitionName))
-						{
-							Result = *FoundGeneratedParameterType;
-						}
-						else
-						{
-							FFormatNamedArguments Args;
-							Args.Add(TEXT("DefinitionName"), FText::FromString(DefinitionName));
-							MessageLog->LogInfo(FText::Format(LOCTEXT("CannotResolveType", "ResolveType (object) failed to find a matching type for definition \"{DefinitionName}\", creating a new one."), Args), FWebAPIOpenAPIProvider::LogName);
-
-							// @bug: where is the model created for this?
-							// @todo: should this always be of type model, vs parameter?
-							Result = OutputSchema->TypeRegistry->GetOrMakeGeneratedType(
-								EWebAPISchemaType::Model,
-								NameTransformer(DefinitionName),
-								{},
-								TEXT("F"));
-							Result->JsonType = UWebAPIStaticTypeRegistry::ToFromJsonType;
-						}
-					}
-				}
+		static TMap<FString, TObjectPtr<UWebAPITypeInfo>> JsonTypeValueToTypeInfo = {
+			{TEXT("None"), StaticTypeRegistry->Void},
+			{TEXT("Null"), StaticTypeRegistry->Nullptr},
+			{TEXT("String"), StaticTypeRegistry->String},
+			{TEXT("Boolean"), StaticTypeRegistry->Boolean},
+			{TEXT("Number"), StaticTypeRegistry->Int32},
+			{TEXT("Array"), StaticTypeRegistry->String},
+			{TEXT("Object"), StaticTypeRegistry->Object},
+		};
 
-				check(!Result.IsNull());
-				
-				return Result;
-			}
+		TObjectPtr<UWebAPITypeInfo> Result = nullptr;
 
-			template <>
-			TSharedPtr<OpenAPI::V3::FSchemaObject> FWebAPIOpenAPISchemaConverter::ResolveReference<OpenAPI::V3::FSchemaObject>(const FString& InDefinitionName)
+		// If a definition name is supplied, try to find it first
+		if ((InType == TEXT("Object") || InType == TEXT("Array")) && !InDefinitionName.IsEmpty())
+		{
+			if (const TObjectPtr<UWebAPITypeInfo>* FoundTypeInfo = OutputSchema->TypeRegistry->FindGeneratedType(EWebAPISchemaType::Model, InDefinitionName))
 			{
-				if(!InputSchema->Components.IsValid())
-				{
-					return nullptr;
-				}
-				
-				if (const Json::TJsonReference<OpenAPI::V3::FSchemaObject>* FoundDefinition = InputSchema->Components->Schemas.Find(InDefinitionName))
-				{
-					if(!FoundDefinition->IsSet())
-					{
-						return ResolveReference<OpenAPI::V3::FSchemaObject>(FoundDefinition->GetLastPathSegment());
-					}
-					return FoundDefinition->GetShared();
-				}
-
-				return nullptr;
+				Result = *FoundTypeInfo;
 			}
+		}
 
-			template <>
-			TSharedPtr<OpenAPI::V3::FParameterObject> FWebAPIOpenAPISchemaConverter::ResolveReference<OpenAPI::V3::FParameterObject>(const FString& InDefinitionName)
+		// If not found above
+		if (!Result)
+		{
+			// Try specific types
+			if (const TObjectPtr<UWebAPITypeInfo>& FoundMappedTypeInfo = ResolveMappedType(InFormat.IsEmpty() ? InType : InFormat))
 			{
-				if(!InputSchema->Components.IsValid())
-				{
-					return nullptr;
-				}
-				
-				if (const Json::TJsonReference<OpenAPI::V3::FParameterObject>* FoundDefinition = InputSchema->Components->Parameters.Find(InDefinitionName))
-				{
-					if(!FoundDefinition->IsSet())
-					{
-						return ResolveReference<OpenAPI::V3::FParameterObject>(FoundDefinition->GetLastPathSegment());
-					}
-					return FoundDefinition->GetShared();
-				}
-
-				return nullptr;
+				return FoundMappedTypeInfo;
 			}
-
-			template <>
-			TSharedPtr<OpenAPI::V3::FResponseObject> FWebAPIOpenAPISchemaConverter::ResolveReference<V3::FResponseObject>(const FString& InDefinitionName)
+			// Fallback to basic types
+			else if (const TObjectPtr<UWebAPITypeInfo>* FoundTypeInfo = JsonTypeValueToTypeInfo.Find(InType))
 			{
-				if(!InputSchema->Components.IsValid())
-				{
-					return nullptr;
-				}
-				
-				if (const Json::TJsonReference<OpenAPI::V3::FResponseObject>* FoundDefinition = InputSchema->Components->Responses.Find(InDefinitionName))
-				{
-					if(!FoundDefinition->IsSet())
-					{
-						return ResolveReference<OpenAPI::V3::FResponseObject>(FoundDefinition->GetLastPathSegment());
-					}
-					return FoundDefinition->GetShared();
-				}
-
-				return nullptr;
+				Result = *FoundTypeInfo;
 			}
 
-			template <>
-			TSharedPtr<OpenAPI::V3::FSecuritySchemeObject> FWebAPIOpenAPISchemaConverter::ResolveReference<V3::FSecuritySchemeObject>(const FString& InDefinitionName)
+			// If it's not a built-in type
+			if (Result != nullptr && !Result->bIsBuiltinType)
 			{
-				if(!InputSchema->Components.IsValid())
-				{
-					return nullptr;
-				}
-				
-				if (const Json::TJsonReference<OpenAPI::V3::FSecuritySchemeObject>* FoundDefinition = InputSchema->Components->SecuritySchemes.Find(InDefinitionName))
+				// Duplicate it with the provided definition name 
+				if (!InDefinitionName.IsEmpty())
+				{
+					// @todo: what if it's an enum? how do you know at this point if it's an enum or struct?
+					// Allow prefix to be set later depending on this?
+					Result = OutputSchema->TypeRegistry->GetOrMakeGeneratedType(EWebAPISchemaType::Model,
+						InDefinitionName,
+						InDefinitionName,
+						Result);
+					Result->Prefix = TEXT("F");
+					ensure(!Result->Name.IsEmpty());
+				}
+				// otherwise duplicate it as an unnamed/partially resolved type 
+				else
 				{
-					if(!FoundDefinition->IsSet())
-					{
-						return ResolveReference<OpenAPI::V3::FSecuritySchemeObject>(FoundDefinition->GetLastPathSegment());
-					}
-					return FoundDefinition->GetShared();
+					// @todo: prevent!
+					Result = Result->Duplicate(OutputSchema->TypeRegistry);
 				}
-
-				return nullptr;
 			}
+		}
 
-			template <>
-			TSharedPtr<OpenAPI::V3::FRequestBodyObject> FWebAPIOpenAPISchemaConverter::ResolveReference<V3::FRequestBodyObject>(const FString& InDefinitionName)
-			{
-				if (const Json::TJsonReference<OpenAPI::V3::FRequestBodyObject>* FoundDefinition = InputSchema->Components->RequestBodies.Find(InDefinitionName))
-				{
-					if(!FoundDefinition->IsSet())
-					{
-						return ResolveReference<OpenAPI::V3::FRequestBodyObject>(FoundDefinition->GetLastPathSegment());
-					}
-					return FoundDefinition->GetShared();
-				}
-
-				return nullptr;
-			}
+		return Result;
+	}
 
-			FWebAPINameVariant FWebAPIOpenAPISchemaConverter::ResolvePropertyName(const TObjectPtr<UWebAPIProperty>& InProperty, const FWebAPITypeNameVariant& InPotentialName, const TOptional<bool>& bInIsArray)
-			{
-				check(InProperty);
+	TObjectPtr<UWebAPITypeInfo> FWebAPIOpenAPISchemaConverter::GetTypeForContentType(const FString& InContentType)
+	{
+		const TObjectPtr<UWebAPIStaticTypeRegistry> StaticTypeRegistry = IWebAPIEditorModuleInterface::Get().GetStaticTypeRegistry();
+		
+		static TMap<FString, TObjectPtr<UWebAPITypeInfo>> ContentTypeToTypeInfo = {
+			{TEXT("application/json"), StaticTypeRegistry->JsonObject},
+			{TEXT("application/xml"), StaticTypeRegistry->String},
+			{TEXT("text/plain"), StaticTypeRegistry->String}
+		};
+		
+		if(const TObjectPtr<UWebAPITypeInfo>* FoundTypeInfo = ContentTypeToTypeInfo.Find(InContentType.ToLower()))
+		{
+			return *FoundTypeInfo;
+		}
 
-				const bool bIsArray = bInIsArray.Get(InProperty->bIsArray);
+		// For all other cases, use string
+		return StaticTypeRegistry->String;
+	}
 
-				// If it's an array, it may be called the generic "Values", try to find a better name
-				if(bIsArray && InProperty->Name == ProviderSettings.GetDefaultArrayPropertyName())
-				{
-					if(InPotentialName.IsValid())
-					{
-						FWebAPINameInfo NameInfo = FWebAPINameInfo(ProviderSettings.Pluralize(InPotentialName.ToString(true)), InProperty->Name.GetJsonName());
-						return NameInfo;
-					}
-					
-					if(InProperty->Type.HasTypeInfo() && !InProperty->Type.TypeInfo->bIsBuiltinType)
-					{
-						FWebAPINameInfo NameInfo = FWebAPINameInfo(ProviderSettings.Pluralize(InProperty->Type.ToString(true)), InProperty->Name.GetJsonName());
-						return NameInfo;
-					}
-				}
-				else if(InProperty->Name == ProviderSettings.GetDefaultPropertyName())
-				{
-					if(InPotentialName.IsValid())
-					{
-						FWebAPINameInfo NameInfo = FWebAPINameInfo(InPotentialName.ToString(true), InProperty->Name.GetJsonName());
-						return NameInfo;
-					}
-					
-					if(InProperty->Type.HasTypeInfo() && !InProperty->Type.TypeInfo->bIsBuiltinType)
-					{
-						FWebAPINameInfo NameInfo = FWebAPINameInfo(InProperty->Type.ToString(true), InProperty->Name.GetJsonName());
-						return NameInfo;
-					}
-				}
+	FWebAPINameVariant FWebAPIOpenAPISchemaConverter::ResolvePropertyName(const TObjectPtr<UWebAPIProperty>& InProperty, const FWebAPITypeNameVariant& InPotentialName, const TOptional<bool>& bInIsArray)
+	{
+		check(InProperty);
 
-				return InProperty->Name;
-			}
+		const bool bIsArray = bInIsArray.Get(InProperty->bIsArray);
 
-			template <typename ObjectType>
-			TSharedPtr<ObjectType> FWebAPIOpenAPISchemaConverter::ResolveReference(const Json::TJsonReference<ObjectType>& InJsonReference, FString& OutDefinitionName, bool bInCheck)
+		// If it's an array, it may be called the generic "Values", try to find a better name
+		if(bIsArray && InProperty->Name == ProviderSettings.GetDefaultArrayPropertyName())
+		{
+			if(InPotentialName.IsValid())
 			{
-				if (InJsonReference.IsSet())
-				{
-					return InJsonReference.GetShared();
-				}
-
-				if(!InJsonReference.IsValid())
-				{
-					return nullptr;
-				}
-				
-				FString DefinitionName = InJsonReference.GetLastPathSegment();
-				if (TSharedPtr<ObjectType> FoundDefinition = ResolveReference<ObjectType>(DefinitionName))
-				{
-					OutDefinitionName = DefinitionName; // Only set if found 
-					return FoundDefinition;
-				}
-
-				if(bInCheck)
-				{
-					FFormatNamedArguments Args;
-					Args.Add(TEXT("ReferencePath"), FText::FromString(InJsonReference.GetPath()));
-					MessageLog->LogWarning(FText::Format(LOCTEXT("CannotResolveJsonReference", "Couldn't resolve JsonReference \"{ReferencePath}\""), Args),	FWebAPIOpenAPIProvider::LogName);
-				}
-
-				return nullptr;
+				FWebAPINameInfo NameInfo = FWebAPINameInfo(ProviderSettings.Pluralize(InPotentialName.ToString(true)), InProperty->Name.GetJsonName());
+				return NameInfo;
 			}
-
-			template <typename ObjectType>
-			TSharedPtr<ObjectType> FWebAPIOpenAPISchemaConverter::ResolveReference(const UE::Json::TJsonReference<ObjectType>& InJsonReference, bool bInCheck)
+			
+			if(InProperty->Type.HasTypeInfo() && !InProperty->Type.TypeInfo->bIsBuiltinType)
 			{
-				FString DefinitionName;
-				return ResolveReference(InJsonReference, DefinitionName, bInCheck);
+				FWebAPINameInfo NameInfo = FWebAPINameInfo(ProviderSettings.Pluralize(InProperty->Type.ToString(true)), InProperty->Name.GetJsonName());
+				return NameInfo;
 			}
-
-			template <>
-			bool FWebAPIOpenAPISchemaConverter::IsArray(const TSharedPtr<UE::WebAPI::OpenAPI::V3::FSchemaObject>& InSchema)
+		}
+		else if(InProperty->Name == ProviderSettings.GetDefaultPropertyName())
+		{
+			if(InPotentialName.IsValid())
 			{
-				return InSchema->Type.Get(TEXT(""))	== TEXT("array");
+				FWebAPINameInfo NameInfo = FWebAPINameInfo(InPotentialName.ToString(true), InProperty->Name.GetJsonName());
+				return NameInfo;
 			}
-
-			template <>
-			bool FWebAPIOpenAPISchemaConverter::IsArray(const TSharedPtr<UE::WebAPI::OpenAPI::V3::FParameterObject>& InSchema)
+			
+			if(InProperty->Type.HasTypeInfo() && !InProperty->Type.TypeInfo->bIsBuiltinType)
 			{
-				FString FoundDefinitionName;
-				return InSchema->Schema.IsSet()
-					? IsArray(ResolveReference(InSchema->Schema.GetValue(), FoundDefinitionName))
-					: false;
+				FWebAPINameInfo NameInfo = FWebAPINameInfo(InProperty->Type.ToString(true), InProperty->Name.GetJsonName());
+				return NameInfo;
 			}
+		}
 
-			template <typename SchemaType, typename ModelType>
-			bool FWebAPIOpenAPISchemaConverter::ConvertModelBase(const TSharedPtr<SchemaType>& InSchema,
-																const TObjectPtr<ModelType>& OutModel)
-			{
-				static_assert(TIsDerivedFrom<SchemaType, OpenAPI::V3::FSchemaObjectBase>::Value, "Type is not derived from OpenAPI::V3::FSchemaObjectBase.");
-/*
-				OutModel->Description = InSchema->Description.Get(TEXT(""));
-
-				SET_OPTIONAL(InSchema->bRequired, OutModel->bIsRequired)
-				
-				SET_OPTIONAL_FLAGGED(InSchema->Minimum, OutModel->MinimumValue, OutModel->bUseMinimumValue);
-				SET_OPTIONAL_FLAGGED(InSchema->Maximum, OutModel->MaximumValue, OutModel->bUseMaximumValue);
-
-				if(!OutModel->bUseMinimumValue)
-				{
-					SET_OPTIONAL_FLAGGED(InSchema->MinLength, OutModel->MinimumValue, OutModel->bUseMinimumValue);
-				}
+		return InProperty->Name;
+	}
 
-				if(!OutModel->bUseMaximumValue)
-				{
-					SET_OPTIONAL_FLAGGED(InSchema->MaxLength, OutModel->MaximumValue, OutModel->bUseMaximumValue);					
-				}
-				
-				SET_OPTIONAL_FLAGGED(InSchema->Pattern, OutModel->Pattern, OutModel->bUsePattern);
-				*/
-				return true;
-			}
 
-			template <typename SchemaType>
-			TObjectPtr<UWebAPIEnum> FWebAPIOpenAPISchemaConverter::ConvertEnum(const TSharedPtr<SchemaType>& InSrcSchema, const FWebAPITypeNameVariant& InEnumTypeName) const
-			{
-				static_assert(TIsDerivedFrom<SchemaType, OpenAPI::V3::FSchemaObjectBase>::Value, "Type is not derived from OpenAPI::V3::FSchemaObjectBase.");
 
-				FWebAPITypeNameVariant EnumTypeName;
-				if(InEnumTypeName.IsValid())
-				{
-					EnumTypeName = InEnumTypeName;
-				}
-				else
-				{
-					const FString EnumName = InSrcSchema->Name;
-					check(!EnumName.IsEmpty());
-
-					const TObjectPtr<UWebAPIStaticTypeRegistry> StaticTypeRegistry = IWebAPIEditorModuleInterface::Get().GetStaticTypeRegistry();
-					EnumTypeName = OutputSchema->TypeRegistry->GetOrMakeGeneratedType(
-						EWebAPISchemaType::Model,
-						NameTransformer(EnumName),
-						EnumName,
-						StaticTypeRegistry->Enum);
-				}
+	bool FWebAPIOpenAPISchemaConverter::ConvertProperty(const FWebAPITypeNameVariant& InModelName, const FWebAPINameVariant& InPropertyName, const FWebAPITypeNameVariant& InPropertyTypeName, const TObjectPtr<UWebAPIProperty>& OutProperty)
+	{
+		OutProperty->Type = InPropertyTypeName;
+		check(OutProperty->Type.IsValid());
 
-				const TObjectPtr<UWebAPIEnum>& DstEnum = OutputSchema->AddEnum(EnumTypeName.TypeInfo.Get());
-				DstEnum->Name = EnumTypeName;
-				DstEnum->Description = InSrcSchema->Description.Get(TEXT(""));
+		OutProperty->Name = FWebAPINameInfo(NameTransformer(InPropertyName.ToString()), InPropertyName.GetJsonName(), OutProperty->Type);
+		OutProperty->Name = ResolvePropertyName(OutProperty, OutProperty->Type.ToString(true), {});
+		OutProperty->bIsRequired = false;
+		OutProperty->BindToTypeInfo();
 
-				const TObjectPtr<UWebAPIModelBase> ModelBase = Cast<UWebAPIModelBase>(DstEnum);
-				if (!ConvertModelBase(InSrcSchema, ModelBase))
-				{
-					return nullptr;
-				}
+		return true;
+	}
 
-				for (const FString& SrcEnumValue : InSrcSchema->Enum.GetValue())
-				{
-					const TObjectPtr<UWebAPIEnumValue> DstEnumValue = DstEnum->AddValue();
-					DstEnumValue->Name.NameInfo.Name = NameTransformer(SrcEnumValue);
-					DstEnumValue->Name.NameInfo.JsonName = SrcEnumValue;
-				}
+	bool FWebAPIOpenAPISchemaConverter::ConvertOperationParameter(
+		const FWebAPINameVariant& InParameterName,
+		const TSharedPtr<OpenAPI::V3::FParameterObject>& InParameter,
+		const FString& InDefinitionName,
+		const TObjectPtr<UWebAPIOperationParameter>& OutParameter)
+	{
+		// Will get schema or create if it doesn't exist (but will be empty)
+		const TSharedPtr<V3::FSchemaObject> ParameterSchema = ResolveReference(InParameter->Schema.Get({}));
+
+		static TMap<FString, EWebAPIParameterStorage> InToStorage = {
+			{TEXT("query"), EWebAPIParameterStorage::Query},
+			{TEXT("header"), EWebAPIParameterStorage::Header},
+			{TEXT("path"), EWebAPIParameterStorage::Path},
+			{TEXT("formData"), EWebAPIParameterStorage::Body},
+			{TEXT("body"), EWebAPIParameterStorage::Body}
+		};
 
-				DstEnum->BindToTypeInfo();
-				
-				return DstEnum;
-			}
+		OutParameter->Storage = InToStorage[InParameter->In];
+		
+		if (ParameterSchema.IsValid())
+		{
+			OutParameter->Type = ResolveType(ParameterSchema, InDefinitionName);
+		}
+		else
+		{
+			// @todo:
+			// OutParameter->Type = ResolveType(InParameter, InDefinitionName, GetDefaultJsonTypeForStorage(OutParameter->Storage));
+		}
 
-			template <>
-			bool FWebAPIOpenAPISchemaConverter::ConvertProperty<OpenAPI::V3::FParameterObject>(
-				const FWebAPITypeNameVariant& InModelName,
-				const FWebAPINameVariant& InPropertyName,
-				const TSharedPtr<OpenAPI::V3::FParameterObject>& InParameter,
-				const FString& InDefinitionName,
-				const TObjectPtr<UWebAPIProperty>& OutProperty)
+		OutParameter->bIsArray = IsArray(InParameter);
+		if (!OutParameter->Type.HasTypeInfo())
+		{
+			if(InParameter->Schema.IsSet())
 			{
-				const FString DefinitionName = !InDefinitionName.IsEmpty()
-												? InDefinitionName
-												: ProviderSettings.MakeNestedPropertyTypeName(InModelName, InPropertyName);
-
-				const TObjectPtr<UWebAPIModelBase> ModelBase = OutProperty;
-				if (!ConvertModelBase(InParameter, ModelBase))
-				{
-					return false;
-				}
-
-				OutProperty->bIsArray = IsArray(InParameter);
-				
-				OutProperty->Name = FWebAPINameInfo(NameTransformer(InPropertyName.ToString()), InPropertyName.GetJsonName(), OutProperty->Type);				
-				OutProperty->Name = ResolvePropertyName(OutProperty, InModelName, {});
-				OutProperty->bIsRequired = InParameter->bRequired.Get(false);
-				OutProperty->BindToTypeInfo();
-
-				if(InParameter->Schema.IsSet())
-				{
-					const TSharedPtr<V3::FSchemaObject> Schema = ResolveReference(InParameter->Schema.GetValue());
-					OutProperty->Type = ResolveType(Schema);
-					
-					// Add enum as it's own model, and reference it as this properties type
-					if (InParameter->Schema.GetValue()->Enum.IsSet() && !InParameter->Schema.GetValue()->Enum->IsEmpty())
-					{
-						const FWebAPITypeNameVariant EnumTypeName = OutProperty->Type;
-						EnumTypeName.TypeInfo->SetNested(InModelName);
-						FString EnumName = OutProperty->Name.ToString(true);
-
-						// Only make nested name if the model and property name aren't the same, otherwise you get "NameName"!
-						if(EnumName != InModelName)
-						{
-							EnumName = ProviderSettings.MakeNestedPropertyTypeName(InModelName, OutProperty->Name.ToString(true));
-						}
-						EnumTypeName.TypeInfo->SetName(EnumName);
-
-						const TObjectPtr<UWebAPIEnum>& Enum = ConvertEnum(Schema, EnumTypeName);
-
-						const FText LogMessage = FText::FormatNamed(
-							LOCTEXT("AddedImplicitEnumForPropertyOfModel", "Implicit enum created for property \"{PropertyName}\" of model \"{ModelName}\"."),
-							TEXT("PropertyName"), FText::FromString(*InPropertyName.ToString(true)),
-							TEXT("ModelName"), FText::FromString(*InModelName.ToString(true)));
-						MessageLog->LogInfo(LogMessage, FWebAPIOpenAPIProvider::LogName);
-					
-						Enum->Name.TypeInfo->DebugString += LogMessage.ToString();
-						Enum->Name.TypeInfo->JsonName = InPropertyName.GetJsonName();
-
-						OutProperty->Type = Enum->Name;
-						OutProperty->Type.TypeInfo->Model = Enum;
-					}
-				}
-
-				return true;
+				// @todo
+				/* 
+				OutParameter->Type = ResolveType(InParameter->Schema->Get(GetDefaultJsonTypeForStorage(OutParameter->Storage)),
+					TEXT(""), //InParameter->Format.Get(TEXT("")),
+					InDefinitionName,
+					ParameterSchema.GetShared());
+					*/
 			}
+		}
 
-			template <typename SchemaType>
-			bool FWebAPIOpenAPISchemaConverter::ConvertProperty(const FWebAPITypeNameVariant& InModelName,
-																const FWebAPINameVariant& InPropertyName,
-																const TSharedPtr<SchemaType>& InSchema,
-																const FString& InDefinitionName,
-																const TObjectPtr<UWebAPIProperty>& OutProperty)
-			{
-				static_assert(TIsDerivedFrom<SchemaType, OpenAPI::V3::FSchemaObjectBase>::Value, "Type is not derived from OpenAPI::V3::FSchemaObjectBase.");
-
-				FString DefinitionName = !InDefinitionName.IsEmpty()
-										? InDefinitionName
-										: ProviderSettings.MakeNestedPropertyTypeName(InModelName, InPropertyName);
-
-				TSharedPtr<UE::WebAPI::OpenAPI::V3::FSchemaObject> ItemSchema = InSchema;
-				if(InSchema->Items.IsSet())
-				{
-					if(!InSchema->Items->GetPath().IsEmpty())
-					{
-						ItemSchema = ResolveReference(InSchema->Items.GetValue(), DefinitionName);
-						DefinitionName = InSchema->Items->GetLastPathSegment();						
-					}
-				}
-				else
-				{
-					ItemSchema = InSchema;
-				}
-
-				OutProperty->Type = ResolveType<SchemaType>(ItemSchema, DefinitionName);
-				check(OutProperty->Type.IsValid());
-
-				const TObjectPtr<UWebAPIModelBase> ModelBase = OutProperty;
-				if (!ConvertModelBase(ItemSchema, ModelBase))
-				{
-					return false;
-				}
-
-				OutProperty->bIsArray = IsArray(InSchema);
-				
-				OutProperty->Name = FWebAPINameInfo(NameTransformer(InPropertyName.ToString()), InPropertyName.GetJsonName(), OutProperty->Type);
-				OutProperty->Name = ResolvePropertyName(OutProperty, OutProperty->Type.ToString(true), {});
-				OutProperty->bIsRequired = InSchema->bRequired.Get(false);
-				OutProperty->BindToTypeInfo();
-
-				// Add enum as it's own model, and reference it as this properties type
-				if (InSchema->Enum.IsSet() && !InSchema->Enum->IsEmpty())
-				{
-					const FWebAPITypeNameVariant EnumTypeName = OutProperty->Type;
-					EnumTypeName.TypeInfo->SetName(ProviderSettings.MakeNestedPropertyTypeName(InModelName, InPropertyName));
-
-					const TObjectPtr<UWebAPIEnum>& Enum = ConvertEnum(InSchema, EnumTypeName);
+		const TObjectPtr<UWebAPIModelBase> ModelBase = Cast<UWebAPIModelBase>(OutParameter);
+		if (!ConvertModelBase(InParameter, ModelBase))
+		{
+			return false;
+		}
 
-					const FText LogMessage = FText::FormatNamed(
-						LOCTEXT("AddedImplicitEnumForPropertyOfModel", "Implicit enum created for property \"{PropertyName}\" of model \"{ModelName}\"."),
-						TEXT("PropertyName"), FText::FromString(*InPropertyName.ToString(true)),
-						TEXT("ModelName"), FText::FromString(*InModelName.ToString(true)));
-					MessageLog->LogInfo(LogMessage, FWebAPIOpenAPIProvider::LogName);
+		if(!OutParameter->Type.TypeInfo->bIsBuiltinType)
+		{
+			// @todo: make generated type here
+			auto x = OutParameter->Type;
+			// ConvertEnum(InParameter)
+		}
 
-					Enum->Name.TypeInfo->DebugString += LogMessage.ToString();
-					
-					Enum->Name.TypeInfo->JsonName = InPropertyName.GetJsonName();
-					Enum->Name.TypeInfo->SetNested(InModelName);
+		OutParameter->Name = FWebAPINameInfo(InParameterName.ToString(true), InParameterName.GetJsonName(), OutParameter->Type);
 
-					OutProperty->Type = Enum->Name;
-					OutProperty->Type.TypeInfo->Model = Enum;
-				}
-				// Add struct as it's own model, and reference it as this properties type
-				else if (OutProperty->Type.ToString(true).IsEmpty())
-				{
-					check(OutProperty->Type.HasTypeInfo());
-					
-					const FWebAPITypeNameVariant ModelTypeName = OutProperty->Type;
-					ModelTypeName.TypeInfo->SetName(ProviderSettings.MakeNestedPropertyTypeName(InModelName, InPropertyName));
-					ModelTypeName.TypeInfo->Prefix = TEXT("F");
-					ModelTypeName.TypeInfo->SetNested(InModelName);
-
-					const TObjectPtr<UWebAPIModel>& Model = OutputSchema->AddModel<UWebAPIModel>(ModelTypeName.HasTypeInfo() ? ModelTypeName.TypeInfo.Get() : nullptr);
-					ConvertModel(InSchema, {}, Model);
-
-					const FText LogMessage = FText::FormatNamed(
-						LOCTEXT("AddedImplicitModelForPropertyOfModel", "Implicit model created for property \"{PropertyName}\" of model \"{ModelName}\"."),
-						TEXT("PropertyName"), FText::FromString(*InPropertyName.ToString(true)),
-						TEXT("ModelName"), FText::FromString(*InModelName.ToString(true)));
-					MessageLog->LogInfo(LogMessage, FWebAPIOpenAPIProvider::LogName);
-					
-					Model->Name.TypeInfo->DebugString += LogMessage.ToString();
-					
-					Model->Name.TypeInfo->JsonName = InPropertyName.GetJsonName();
-					Model->Name.TypeInfo->JsonType = UWebAPIStaticTypeRegistry::ToFromJsonType;
-
-					OutProperty->Type = Model->Name;
-					OutProperty->Type.TypeInfo->Model = Model;
-				}
-				
-				OutProperty->Name = ResolvePropertyName(OutProperty, InModelName);
+		// Add struct as it's own model, and reference it as this properties type
+		if (OutParameter->Type.ToString(true).IsEmpty())
+		{
+			const FWebAPITypeNameVariant ModelTypeName = OutParameter->Type;
+			ModelTypeName.TypeInfo->SetName(InDefinitionName);
+			ModelTypeName.TypeInfo->Prefix = TEXT("F");
+
+			const TObjectPtr<UWebAPIModel>& Model = OutputSchema->AddModel<UWebAPIModel>(
+				ModelTypeName.HasTypeInfo() ? ModelTypeName.TypeInfo.Get() : nullptr);
+			ConvertModel(InParameter->Schema.GetValue().GetShared(), {}, Model);
+
+			const FText LogMessage = FText::FormatNamed(
+				LOCTEXT("AddedImplicitModelForParameterOfOperation", "Implicit model created for parameter \"{ParameterName}\" of operation \"{OperationName}\"."),
+				TEXT("ParameterName"), FText::FromString(*InParameterName.ToString(true)),
+				TEXT("OperationName"), FText::FromString(TEXT("?")));
+			MessageLog->LogInfo(LogMessage, FWebAPIOpenAPIProvider::LogName);
+			
+			Model->Name.TypeInfo->DebugString += LogMessage.ToString();
+			
+			Model->Name.TypeInfo->JsonName = InParameterName.GetJsonName();
+			Model->Name.TypeInfo->JsonType = UWebAPIStaticTypeRegistry::ToFromJsonType;
 
-				return true;
-			}
+			OutParameter->Type = Model->Name;
+			OutParameter->Type.TypeInfo->Model = Model;
+		}
 
-			bool FWebAPIOpenAPISchemaConverter::ConvertProperty(const FWebAPITypeNameVariant& InModelName, const FWebAPINameVariant& InPropertyName, const FWebAPITypeNameVariant& InPropertyTypeName, const TObjectPtr<UWebAPIProperty>& OutProperty)
+		// Special case for "body" parameters
+		if (InParameter->In == TEXT("body"))
+		{
+			if(OutParameter->Name.ToString(true) == TEXT("Body"))
 			{
-				OutProperty->Type = InPropertyTypeName;
-				check(OutProperty->Type.IsValid());
-
-				OutProperty->Name = FWebAPINameInfo(NameTransformer(InPropertyName.ToString()), InPropertyName.GetJsonName(), OutProperty->Type);
-				OutProperty->Name = ResolvePropertyName(OutProperty, OutProperty->Type.ToString(true), {});
-				OutProperty->bIsRequired = false;
-				OutProperty->BindToTypeInfo();
-
-				return true;
+				OutParameter->Name = InDefinitionName;
 			}
 
-			template <typename SchemaType>
-			TObjectPtr<UWebAPIProperty> FWebAPIOpenAPISchemaConverter::ConvertProperty(const TSharedPtr<SchemaType>& InSrcSchema, const TObjectPtr<UWebAPIModel>& InModel, const FWebAPINameVariant& InPropertyName, const FString& InDefinitionName)
+			if (OutParameter->bIsArray)
 			{
-				static_assert(TIsDerivedFrom<SchemaType, OpenAPI::V3::FSchemaObjectBase>::Value, "Type is not derived from OpenAPI::V3::FSchemaObjectBase.");
-
-				const FWebAPINameVariant SrcPropertyName = InPropertyName;
-				FString SrcPropertyDefinitionName = ProviderSettings.MakeNestedPropertyTypeName(*InModel->Name.ToString(true), SrcPropertyName);
-				const TSharedPtr<OpenAPI::V3::FSchemaObject> SrcPropertyValue = ResolveReference(InSrcSchema, SrcPropertyDefinitionName);
-
-				if (SrcPropertyValue)
-				{
-					const TObjectPtr<UWebAPIProperty>& DstProperty = InModel->Properties.Add_GetRef(NewObject<UWebAPIProperty>(InModel));
-					ConvertProperty(InModel->Name,
-						FWebAPINameInfo(NameTransformer(SrcPropertyName), SrcPropertyName.GetJsonName(), InModel->Name),
-						SrcPropertyValue,
-						SrcPropertyDefinitionName,
-						DstProperty);
-					
-					return DstProperty;
-				}
-
-				return nullptr;
+				FString Name = OutParameter->Name.ToString(true);
+				Name = ProviderSettings.Singularize(Name);
+				Name = ProviderSettings.Pluralize(Name);
+				OutParameter->Name = Name;
 			}
+		}
 
-			template <typename SchemaType>
-			bool FWebAPIOpenAPISchemaConverter::ConvertModel(const TSharedPtr<SchemaType>& InSrcSchema, const FWebAPITypeNameVariant& InModelTypeName, const TObjectPtr<UWebAPIModel>& OutModel)
-			{
-				static_assert(TIsDerivedFrom<SchemaType, OpenAPI::V3::FSchemaObjectBase>::Value, "Type is not derived from OpenAPI::V3::FSchemaObjectBase.");
-
-				FWebAPITypeNameVariant ModelTypeName;
-				if(InModelTypeName.IsValid())
-				{
-					ModelTypeName = InModelTypeName;
-				}
-				else if(OutModel->Name.HasTypeInfo())
-				{
-					ModelTypeName = OutModel->Name;
-				}
-				else
-				{
-					const FString ModelName = InSrcSchema->Name;
-					check(!ModelName.IsEmpty());
-
-					ModelTypeName = OutputSchema->TypeRegistry->GetOrMakeGeneratedType(
-						EWebAPISchemaType::Model,
-						NameTransformer(*ModelName),
-						*ModelName,
-						TEXT("F"));
-				}
-
-				if (!InSrcSchema.IsValid())
-				{
-					FFormatNamedArguments Args;
-					Args.Add(TEXT("ModelName"), FText::FromString(*ModelTypeName.ToString(true)));
-					MessageLog->LogWarning(
-						FText::Format(LOCTEXT("SchemaInvalid",
-								"The schema for model \"{ModelName}\" was invalid/null."),
-							Args),
-						FWebAPIOpenAPIProvider::LogName);
-					return false;
-				}
-
-				const TObjectPtr<UWebAPIModel>& Model = OutModel ? OutModel : OutputSchema->AddModel<UWebAPIModel>(ModelTypeName.TypeInfo.Get());
+		OutParameter->Description = InParameter->Description.Get(TEXT(""));
+		OutParameter->bIsRequired = InParameter->bRequired.Get(false);
+		OutParameter->BindToTypeInfo();
+		
+		// ConvertProperty({}, InParameterName, InParameter, InDefinitionName, OutParameter);
 
-				const TObjectPtr<UWebAPIModelBase> ModelBase = Model;
-				if (!ConvertModelBase(InSrcSchema, ModelBase))
-				{
-					return false;
-				}
+		return true;
+	}
 
-				Model->Name = ModelTypeName;
+	TObjectPtr<UWebAPIParameter> FWebAPIOpenAPISchemaConverter::ConvertParameter(const TSharedPtr<OpenAPI::V3::FParameterObject>& InSrcParameter)
+	{
+		// const FString* ParameterName = InputSchema->Components.Parameters.FindKey(InSrcParameter);
+		// check(ParameterName);
+		const FString* ParameterName = nullptr;
+
+		const FString ParameterJsonName = InSrcParameter->Name.IsEmpty() ? *ParameterName : InSrcParameter->Name;
+
+		const FWebAPITypeNameVariant ParameterTypeName = OutputSchema->TypeRegistry->GetOrMakeGeneratedType(
+			EWebAPISchemaType::Parameter,
+			ProviderSettings.MakeParameterTypeName(NameTransformer(*ParameterName)),
+			ParameterJsonName,
+			TEXT("F"));
+		
+		ParameterTypeName.TypeInfo->Suffix = TEXT("Parameter");
+
+		// Don't use ParameterTypeName->Name, it might have a Parameter specific pre/postfix.
+		FString SrcParameterDefinitionName = NameTransformer(*ParameterName);
+
+		// Will get schema or create if it doesn't exist (but will be empty)
+		const TSharedPtr<OpenAPI::V3::FSchemaObject> SrcParameterSchema = ResolveReference(InSrcParameter->Schema.Get({}), SrcParameterDefinitionName, false);
+
+		const TObjectPtr<UWebAPIParameter> DstParameter = OutputSchema->AddParameter(ParameterTypeName.TypeInfo.Get());
+
+		static TMap<FString, EWebAPIParameterStorage> InToStorage = {
+			{TEXT("query"), EWebAPIParameterStorage::Query},
+			{TEXT("header"), EWebAPIParameterStorage::Header},
+			{TEXT("path"), EWebAPIParameterStorage::Path},
+			{TEXT("formData"), EWebAPIParameterStorage::Body},
+			{TEXT("body"), EWebAPIParameterStorage::Body}
+		};
 
-				if (InSrcSchema->Properties.IsSet() && !InSrcSchema->Properties->IsEmpty())
-				{
-					for (const TTuple<FString, Json::TJsonReference<OpenAPI::V3::FSchemaObject>>& NamePropertyPair : InSrcSchema->Properties.GetValue())
-					{
-						const FString& SrcPropertyName = NamePropertyPair.Key;
-						
-						FString SrcPropertyDefinitionName = ProviderSettings.MakeNestedPropertyTypeName(*ModelTypeName.ToString(true), SrcPropertyName);
-						TSharedPtr<OpenAPI::V3::FSchemaObject> SrcPropertySchema = ResolveReference(NamePropertyPair.Value, SrcPropertyDefinitionName);
-
-						TObjectPtr<UWebAPIProperty>& DstProperty = Model->Properties.Add_GetRef(NewObject<UWebAPIProperty>(Model));
-						ConvertProperty(Model->Name,
-										FWebAPINameInfo(NameTransformer(SrcPropertyName), SrcPropertyName, ModelTypeName),
-										SrcPropertySchema,
-										SrcPropertyDefinitionName,
-										DstProperty);
-					}
-				}
+		DstParameter->Storage = InToStorage[InSrcParameter->In];
 
-				Model->BindToTypeInfo();
+		DstParameter->bIsArray = IsArray(InSrcParameter);
+		if (!DstParameter->Type.HasTypeInfo())
+		{
+			// @todo:
+			/*
+			DstParameter->Type = ResolveType(InSrcParameter->Type.Get(GetDefaultJsonTypeForStorage(DstParameter->Storage)),
+				InSrcParameter->Format.Get(TEXT("")),
+				SrcParameterDefinitionName,
+				SrcParameterSchema);
+				*/
+		}
 
-				return true;
-			}
+		const TObjectPtr<UWebAPIModel> Model = Cast<UWebAPIModel>(DstParameter);
+		if (!ConvertModel(InSrcParameter, ParameterTypeName, Model))
+		{
+			return nullptr;
+		}
 
-			template <>
-			bool FWebAPIOpenAPISchemaConverter::ConvertModel<OpenAPI::V3::FParameterObject>(const TSharedPtr<OpenAPI::V3::FParameterObject>& InSrcSchema, const FWebAPITypeNameVariant& InModelTypeName, const TObjectPtr<UWebAPIModel>& OutModel)
+		// Set the (single) property, will either be a single or array of a model or enum, based on the parameter or it's schema + parameter
+		{
+			FString PropertyName = DstParameter->bIsArray ? ProviderSettings.GetDefaultArrayPropertyName() : ProviderSettings.GetDefaultPropertyName();
+			SrcParameterDefinitionName = ProviderSettings.MakeNestedPropertyTypeName(SrcParameterDefinitionName, PropertyName);
+			
+			TObjectPtr<UWebAPIProperty>& DstProperty = DstParameter->Property = Model->Properties.Add_GetRef(NewObject<UWebAPIProperty>(Model));
+			ConvertProperty(Model->Name,
+				PropertyName,
+				InSrcParameter,
+				SrcParameterDefinitionName,
+				DstProperty);
+
+			// Add enum as it's own model, and reference it as this properties type
+			if (SrcParameterSchema.IsValid() && SrcParameterSchema->Enum.IsSet() && !SrcParameterSchema->Enum->IsEmpty())
 			{
-				FWebAPITypeNameVariant ModelTypeName;
-				if(InModelTypeName.IsValid() && InModelTypeName.HasTypeInfo())
-				{
-					ModelTypeName = InModelTypeName;
-				}
-				else if(OutModel->Name.HasTypeInfo())
-				{
-					ModelTypeName = OutModel->Name;
-				}
-				else
-				{
-					const FString ModelName = InSrcSchema->Name;
-					check(!ModelName.IsEmpty());
-
-					ModelTypeName = OutputSchema->TypeRegistry->GetOrMakeGeneratedType(
-						EWebAPISchemaType::Model,
-						NameTransformer(*ModelName),
-						*ModelName,
-						TEXT("F"));
-				}
+				const FWebAPITypeNameVariant EnumTypeName = DstParameter->Type;
+				EnumTypeName.TypeInfo->SetName(ProviderSettings.MakeParameterTypeName(*ParameterName));
+				EnumTypeName.TypeInfo->SetNested(ParameterTypeName);
 
-				if (!InSrcSchema.IsValid())
-				{
-					FFormatNamedArguments Args;
-					Args.Add(TEXT("ModelName"), FText::FromString(*ModelTypeName.ToString(true)));
-					MessageLog->LogWarning(
-						FText::Format(LOCTEXT("SchemaInvalid",
-								"The schema for model \"{ModelName}\" was invalid/null."),
-							Args),
-						FWebAPIOpenAPIProvider::LogName);
-					return false;
-				}
+				const TObjectPtr<UWebAPIEnum>& Enum = ConvertEnum(SrcParameterSchema, EnumTypeName);
 
-				const TObjectPtr<UWebAPIModel>& Model = OutModel ? OutModel : OutputSchema->AddModel<UWebAPIModel>(ModelTypeName.TypeInfo.Get());
+				const FText LogMessage = FText::FormatNamed(
+					LOCTEXT("AddedImplicitEnumForParameter", "Implicit enum created for parameter \"{Name}\"."),
+					TEXT("Name"), FText::FromString(**ParameterName));
+				MessageLog->LogInfo(LogMessage, FWebAPIOpenAPIProvider::LogName);
 
-				const TObjectPtr<UWebAPIModelBase> ModelBase = Model;
-				if (!ConvertModelBase(InSrcSchema, ModelBase))
-				{
-					return false;
-				}
-
-				Model->Name = ModelTypeName;
-
-				Model->BindToTypeInfo();
+				Enum->Name.TypeInfo->DebugString += LogMessage.ToString();
 
-				return true;
+				DstParameter->Type = Enum->Name;
+				DstParameter->Type.TypeInfo->Model = Enum;
 			}
-
-			template <typename SchemaType>
-			TObjectPtr<UWebAPIModel> FWebAPIOpenAPISchemaConverter::ConvertModel(const TSharedPtr<SchemaType>& InSrcSchema, const FWebAPITypeNameVariant& InModelTypeName)
+			// Add struct as it's own model, and reference it as this properties type
+			else if (DstParameter->Type.ToString(true).IsEmpty())
 			{
-				check(InSrcSchema.IsValid());
-
-				FWebAPITypeNameVariant ModelTypeName;
-				if(InModelTypeName.IsValid() && ModelTypeName.HasTypeInfo())
-				{
-					ModelTypeName = InModelTypeName;
-				}
-				else
-				{
-					const FString ModelName = InSrcSchema->Name;
-					check(!ModelName.IsEmpty());
-
-					ModelTypeName = OutputSchema->TypeRegistry->GetOrMakeGeneratedType(
-						EWebAPISchemaType::Model,
-						NameTransformer(*ModelName),
-						*ModelName,
-						TEXT("F"));
-				}
+				const FWebAPITypeNameVariant ModelTypeName = DstParameter->Name;
+				ModelTypeName.TypeInfo->SetName(SrcParameterDefinitionName);
+				ModelTypeName.TypeInfo->Prefix = TEXT("F");
+				ModelTypeName.TypeInfo->SetNested(ParameterTypeName);
+
+				//const TObjectPtr<UWebAPIModel>& Model = OutputSchema->AddModel<UWebAPIModel>(ModelTypeName.HasTypeInfo() ? ModelTypeName.TypeInfo.Get() : nullptr);
+				ConvertModel(InSrcParameter, {}, DstParameter);
+
+				const FText LogMessage = FText::FormatNamed(
+					LOCTEXT("AddedImplicitModelForParameter", "Implicit model created for parameter \"{ParameterName}\"."),
+					TEXT("ParameterName"), FText::FromString(**ParameterName));
+				MessageLog->LogInfo(LogMessage, FWebAPIOpenAPIProvider::LogName);
 				
-				const TObjectPtr<UWebAPIModel>& Model = OutputSchema->AddModel<UWebAPIModel>(ModelTypeName.TypeInfo.Get());
-				if(ConvertModel(InSrcSchema, ModelTypeName, Model))
-				{
-					return Model;
-				}
-
-				return nullptr;
-			}
-			
-			bool FWebAPIOpenAPISchemaConverter::ConvertOperationParameter(
-				const FWebAPINameVariant& InParameterName,
-				const TSharedPtr<OpenAPI::V3::FParameterObject>& InParameter,
-				const FString& InDefinitionName,
-				const TObjectPtr<UWebAPIOperationParameter>& OutParameter)
-			{
-				// Will get schema or create if it doesn't exist (but will be empty)
-				const TSharedPtr<V3::FSchemaObject> ParameterSchema = ResolveReference(InParameter->Schema.Get({}));
-
-				static TMap<FString, EWebAPIParameterStorage> InToStorage = {
-					{TEXT("query"), EWebAPIParameterStorage::Query},
-					{TEXT("header"), EWebAPIParameterStorage::Header},
-					{TEXT("path"), EWebAPIParameterStorage::Path},
-					{TEXT("formData"), EWebAPIParameterStorage::Body},
-					{TEXT("body"), EWebAPIParameterStorage::Body}
-				};
-
-				OutParameter->Storage = InToStorage[InParameter->In];
+				DstParameter->Name.TypeInfo->DebugString += LogMessage.ToString();
 				
-				if (ParameterSchema.IsValid())
-				{
-					OutParameter->Type = ResolveType(ParameterSchema, InDefinitionName);
-				}
-				else
-				{
-					// @todo:
-					// OutParameter->Type = ResolveType(InParameter, InDefinitionName, GetDefaultJsonTypeForStorage(OutParameter->Storage));
-				}
+				DstParameter->Name.TypeInfo->JsonName = ParameterTypeName.GetJsonName();
+				DstParameter->Name.TypeInfo->JsonType = UWebAPIStaticTypeRegistry::ToFromJsonType;
 
-				OutParameter->bIsArray = IsArray(InParameter);
-				if (!OutParameter->Type.HasTypeInfo())
-				{
-					if(InParameter->Schema.IsSet())
-					{
-						// @todo
-						/* 
-						OutParameter->Type = ResolveType(InParameter->Schema->Get(GetDefaultJsonTypeForStorage(OutParameter->Storage)),
-							TEXT(""), //InParameter->Format.Get(TEXT("")),
-							InDefinitionName,
-							ParameterSchema.GetShared());
-							*/
-					}
-				}
+				// DstParameter->Type = Model->Name;
+				//					DstParameter->Type.TypeInfo->Model = DstParameter;
+			}
+		}
 
-				const TObjectPtr<UWebAPIModelBase> ModelBase = Cast<UWebAPIModelBase>(OutParameter);
-				if (!ConvertModelBase(InParameter, ModelBase))
-				{
-					return false;
-				}
 
-				if(!OutParameter->Type.TypeInfo->bIsBuiltinType)
-				{
-					// @todo: make generated type here
-					auto x = OutParameter->Type;
-					// ConvertEnum(InParameter)
-				}
+		/**
+		const TObjectPtr<UWebAPIModel> Model = Cast<UWebAPIModel>(DstParameter);
+		if (!ConvertModel(InSrcParameter, ParameterTypeName, Model))
+		{
+			return nullptr;
+		}
 
-				OutParameter->Name = FWebAPINameInfo(InParameterName.ToString(true), InParameterName.GetJsonName(), OutParameter->Type);
+		
+		if (!SrcParameterSchema.GetPath().IsEmpty())
+		{
+			SrcParameterDefinitionName = SrcParameterSchema.GetLastPathSegment();
+		}
 
-				// Add struct as it's own model, and reference it as this properties type
-				if (OutParameter->Type.ToString(true).IsEmpty())
-				{
-					const FWebAPITypeNameVariant ModelTypeName = OutParameter->Type;
-					ModelTypeName.TypeInfo->SetName(InDefinitionName);
-					ModelTypeName.TypeInfo->Prefix = TEXT("F");
-
-					const TObjectPtr<UWebAPIModel>& Model = OutputSchema->AddModel<UWebAPIModel>(
-						ModelTypeName.HasTypeInfo() ? ModelTypeName.TypeInfo.Get() : nullptr);
-					ConvertModel(InParameter->Schema.GetValue().GetShared(), {}, Model);
-
-					const FText LogMessage = FText::FormatNamed(
-						LOCTEXT("AddedImplicitModelForParameterOfOperation", "Implicit model created for parameter \"{ParameterName}\" of operation \"{OperationName}\"."),
-						TEXT("ParameterName"), FText::FromString(*InParameterName.ToString(true)),
-						TEXT("OperationName"), FText::FromString(TEXT("?")));
-					MessageLog->LogInfo(LogMessage, FWebAPIOpenAPIProvider::LogName);
-					
-					Model->Name.TypeInfo->DebugString += LogMessage.ToString();
-					
-					Model->Name.TypeInfo->JsonName = InParameterName.GetJsonName();
-					Model->Name.TypeInfo->JsonType = UWebAPIStaticTypeRegistry::ToFromJsonType;
+		if (SrcParameterSchema.IsSet() && SrcParameterSchema->Items.IsSet() && !SrcParameterSchema->Items->GetPath().IsEmpty())
+		{
+			SrcParameterDefinitionName = SrcParameterSchema->Items->GetLastPathSegment();
+		}
 
-					OutParameter->Type = Model->Name;
-					OutParameter->Type.TypeInfo->Model = Model;
-				}
+		if (SrcParameterSchema.IsSet())
+		{
+			DstParameter->Type = ResolveType(SrcParameterSchema.GetShared(), NameTransformer(*ParameterName));
+		}
+		else
+		{
+			DstParameter->Type = ResolveType(InSrcParameter, NameTransformer(*ParameterName));
+		}
+		*/
 
-				// Special case for "body" parameters
-				if (InParameter->In == TEXT("body"))
-				{
-					if(OutParameter->Name.ToString(true) == TEXT("Body"))
-					{
-						OutParameter->Name = InDefinitionName;
-					}
+		// Special case for "body" parameters
+		if (InSrcParameter->In == TEXT("body"))
+		{
+			DstParameter->Name = SrcParameterDefinitionName;
 
-					if (OutParameter->bIsArray)
-					{
-						FString Name = OutParameter->Name.ToString(true);
-						Name = ProviderSettings.Singularize(Name);
-						Name = ProviderSettings.Pluralize(Name);
-						OutParameter->Name = Name;
-					}
-				}
+			if (DstParameter->bIsArray)
+			{
+				DstParameter->Name = ProviderSettings.Pluralize(SrcParameterDefinitionName);
+			}
+		}
 
-				OutParameter->Description = InParameter->Description.Get(TEXT(""));
-				OutParameter->bIsRequired = InParameter->bRequired.Get(false);
-				OutParameter->BindToTypeInfo();
-				
-				// ConvertProperty({}, InParameterName, InParameter, InDefinitionName, OutParameter);
+		DstParameter->BindToTypeInfo();
 
-				return true;
+		//DstParameter->Name = FWebAPINameInfo(NameTransformer(DstParameter->Name), DstParameter->Name.GetJsonName(), DstParameter->Type);
+		DstParameter->Description = InSrcParameter->Description.Get(TEXT(""));
+		DstParameter->bIsRequired = InSrcParameter->bRequired.Get(false);
+		
+		/*
+		if(!DstParameter->Type.TypeInfo->bIsBuiltinType)
+		{
+			if(DstParameter->Type.TypeInfo->IsEnum())
+			{
+				ConvertEnum(InSrcParameter);						
 			}
+			else
+			{
+				unimplemented();						
+			}
+		}
+		*/
 
-			TObjectPtr<UWebAPIParameter> FWebAPIOpenAPISchemaConverter::ConvertParameter(const TSharedPtr<OpenAPI::V3::FParameterObject>& InSrcParameter)
+		/**
+		// Special case - if the name is "body", set as Body property and not Parameter
+		if (InSrcParameter->Name.Equals(TEXT("body"), ESearchCase::IgnoreCase) && !SrcParameterDefinitionName.IsEmpty())
+		{
+			// Check for existing definition
+			if (const TObjectPtr<UWebAPITypeInfo>* FoundGeneratedType = OutputSchema->TypeRegistry->FindGeneratedType(EWebAPISchemaType::Model, SrcParameterDefinitionName))
 			{
-				// const FString* ParameterName = InputSchema->Components.Parameters.FindKey(InSrcParameter);
-				// check(ParameterName);
-				const FString* ParameterName = nullptr;
+				DstParameter->Model = Cast<UWebAPIModel>((*FoundGeneratedType)->Model.LoadSynchronous());
+				return DstParameter;
+			}
+		}
+		*/
 
-				const FString ParameterJsonName = InSrcParameter->Name.IsEmpty() ? *ParameterName : InSrcParameter->Name;
+		return DstParameter;
+	}
 
-				const FWebAPITypeNameVariant ParameterTypeName = OutputSchema->TypeRegistry->GetOrMakeGeneratedType(
-					EWebAPISchemaType::Parameter,
-					ProviderSettings.MakeParameterTypeName(NameTransformer(*ParameterName)),
-					ParameterJsonName,
-					TEXT("F"));
-				
-				ParameterTypeName.TypeInfo->Suffix = TEXT("Parameter");
+	bool FWebAPIOpenAPISchemaConverter::ConvertRequest(const FWebAPITypeNameVariant& InOperationName, const TSharedPtr<OpenAPI::V3::FOperationObject>& InOperation, const TObjectPtr<UWebAPIOperationRequest>& OutRequest)
+	{
+		// const FWebAPITypeNameVariant RequestTypeName = OutputSchema->TypeRegistry->AddUnnamedType(IWebAPIEditorModuleInterface::Get().GetStaticTypeRegistry()->Object);
+		const FWebAPITypeNameVariant RequestTypeName = OutputSchema->TypeRegistry->GetOrMakeGeneratedType(
+			EWebAPISchemaType::Model,
+			ProviderSettings.MakeRequestTypeName(InOperationName),
+			InOperationName.ToString(true),
+			IWebAPIEditorModuleInterface::Get().GetStaticTypeRegistry()->Object);
+#if WITH_WEBAPI_DEBUG
+		RequestTypeName.TypeInfo->DebugString += TEXT(">ConvertRequest");
+#endif
+		RequestTypeName.TypeInfo->Prefix = TEXT("F");
 
-				// Don't use ParameterTypeName->Name, it might have a Parameter specific pre/postfix.
-				FString SrcParameterDefinitionName = NameTransformer(*ParameterName);
+		OutRequest->Name = RequestTypeName;
 
+		if(!InOperation->Parameters.IsEmpty())
+		{
+			for (Json::TJsonReference<OpenAPI::V3::FParameterObject>& SrcParameter : InOperation->Parameters)
+			{
+				FString SrcParameterDefinitionName;
 				// Will get schema or create if it doesn't exist (but will be empty)
-				const TSharedPtr<OpenAPI::V3::FSchemaObject> SrcParameterSchema = ResolveReference(InSrcParameter->Schema.Get({}), SrcParameterDefinitionName, false);
-
-				const TObjectPtr<UWebAPIParameter> DstParameter = OutputSchema->AddParameter(ParameterTypeName.TypeInfo.Get());
-
-				static TMap<FString, EWebAPIParameterStorage> InToStorage = {
-					{TEXT("query"), EWebAPIParameterStorage::Query},
-					{TEXT("header"), EWebAPIParameterStorage::Header},
-					{TEXT("path"), EWebAPIParameterStorage::Path},
-					{TEXT("formData"), EWebAPIParameterStorage::Body},
-					{TEXT("body"), EWebAPIParameterStorage::Body}
-				};
+				TSharedPtr<OpenAPI::V3::FSchemaObject> SrcParameterSchema = ResolveReference(SrcParameter->Schema.Get({}), SrcParameterDefinitionName);
 
-				DstParameter->Storage = InToStorage[InSrcParameter->In];
-
-				DstParameter->bIsArray = IsArray(InSrcParameter);
-				if (!DstParameter->Type.HasTypeInfo())
-				{
-					// @todo:
-					/*
-					DstParameter->Type = ResolveType(InSrcParameter->Type.Get(GetDefaultJsonTypeForStorage(DstParameter->Storage)),
-						InSrcParameter->Format.Get(TEXT("")),
-						SrcParameterDefinitionName,
-						SrcParameterSchema);
-						*/
-				}
-
-				const TObjectPtr<UWebAPIModel> Model = Cast<UWebAPIModel>(DstParameter);
-				if (!ConvertModel(InSrcParameter, ParameterTypeName, Model))
+				FWebAPINameVariant ParameterName = FWebAPINameInfo(NameTransformer(SrcParameter->Name), SrcParameter->Name);
+				if(SrcParameterDefinitionName.IsEmpty())
 				{
-					return nullptr;
-				}
-
-				// Set the (single) property, will either be a single or array of a model or enum, based on the parameter or it's schema + parameter
-				{
-					FString PropertyName = DstParameter->bIsArray ? ProviderSettings.GetDefaultArrayPropertyName() : ProviderSettings.GetDefaultPropertyName();
-					SrcParameterDefinitionName = ProviderSettings.MakeNestedPropertyTypeName(SrcParameterDefinitionName, PropertyName);
-					
-					TObjectPtr<UWebAPIProperty>& DstProperty = DstParameter->Property = Model->Properties.Add_GetRef(NewObject<UWebAPIProperty>(Model));
-					ConvertProperty(Model->Name,
-						PropertyName,
-						InSrcParameter,
-						SrcParameterDefinitionName,
-						DstProperty);
-
-					// Add enum as it's own model, and reference it as this properties type
-					if (SrcParameterSchema.IsValid() && SrcParameterSchema->Enum.IsSet() && !SrcParameterSchema->Enum->IsEmpty())
-					{
-						const FWebAPITypeNameVariant EnumTypeName = DstParameter->Type;
-						EnumTypeName.TypeInfo->SetName(ProviderSettings.MakeParameterTypeName(*ParameterName));
-						EnumTypeName.TypeInfo->SetNested(ParameterTypeName);
-
-						const TObjectPtr<UWebAPIEnum>& Enum = ConvertEnum(SrcParameterSchema, EnumTypeName);
-
-						const FText LogMessage = FText::FormatNamed(
-							LOCTEXT("AddedImplicitEnumForParameter", "Implicit enum created for parameter \"{Name}\"."),
-							TEXT("Name"), FText::FromString(**ParameterName));
-						MessageLog->LogInfo(LogMessage, FWebAPIOpenAPIProvider::LogName);
-
-						Enum->Name.TypeInfo->DebugString += LogMessage.ToString();
-
-						DstParameter->Type = Enum->Name;
-						DstParameter->Type.TypeInfo->Model = Enum;
-					}
-					// Add struct as it's own model, and reference it as this properties type
-					else if (DstParameter->Type.ToString(true).IsEmpty())
-					{
-						const FWebAPITypeNameVariant ModelTypeName = DstParameter->Name;
-						ModelTypeName.TypeInfo->SetName(SrcParameterDefinitionName);
-						ModelTypeName.TypeInfo->Prefix = TEXT("F");
-						ModelTypeName.TypeInfo->SetNested(ParameterTypeName);
-
-						//const TObjectPtr<UWebAPIModel>& Model = OutputSchema->AddModel<UWebAPIModel>(ModelTypeName.HasTypeInfo() ? ModelTypeName.TypeInfo.Get() : nullptr);
-						ConvertModel(InSrcParameter, {}, DstParameter);
-
-						const FText LogMessage = FText::FormatNamed(
-							LOCTEXT("AddedImplicitModelForParameter", "Implicit model created for parameter \"{ParameterName}\"."),
-							TEXT("ParameterName"), FText::FromString(**ParameterName));
-						MessageLog->LogInfo(LogMessage, FWebAPIOpenAPIProvider::LogName);
-						
-						DstParameter->Name.TypeInfo->DebugString += LogMessage.ToString();
-						
-						DstParameter->Name.TypeInfo->JsonName = ParameterTypeName.GetJsonName();
-						DstParameter->Name.TypeInfo->JsonType = UWebAPIStaticTypeRegistry::ToFromJsonType;
-
-						// DstParameter->Type = Model->Name;
-						//					DstParameter->Type.TypeInfo->Model = DstParameter;
-					}
-				}
-
-	
-				/**
-				const TObjectPtr<UWebAPIModel> Model = Cast<UWebAPIModel>(DstParameter);
-				if (!ConvertModel(InSrcParameter, ParameterTypeName, Model))
-				{
-					return nullptr;
-				}
-
-				
-				if (!SrcParameterSchema.GetPath().IsEmpty())
-				{
-					SrcParameterDefinitionName = SrcParameterSchema.GetLastPathSegment();
-				}
-
-				if (SrcParameterSchema.IsSet() && SrcParameterSchema->Items.IsSet() && !SrcParameterSchema->Items->GetPath().IsEmpty())
-				{
-					SrcParameterDefinitionName = SrcParameterSchema->Items->GetLastPathSegment();
-				}
-
-				if (SrcParameterSchema.IsSet())
-				{
-					DstParameter->Type = ResolveType(SrcParameterSchema.GetShared(), NameTransformer(*ParameterName));
-				}
-				else
-				{
-					DstParameter->Type = ResolveType(InSrcParameter, NameTransformer(*ParameterName));
-				}
-				*/
-
-				// Special case for "body" parameters
-				if (InSrcParameter->In == TEXT("body"))
-				{
-					DstParameter->Name = SrcParameterDefinitionName;
-
-					if (DstParameter->bIsArray)
-					{
-						DstParameter->Name = ProviderSettings.Pluralize(SrcParameterDefinitionName);
-					}
+					SrcParameterDefinitionName = ProviderSettings.MakeNestedPropertyTypeName(InOperationName, ParameterName);
 				}
 
+				const TObjectPtr<UWebAPIOperationParameter> DstParameter = OutRequest->Parameters.Add_GetRef(NewObject<UWebAPIOperationParameter>(OutRequest));
+				ConvertOperationParameter(ParameterName, SrcParameter.GetShared(), SrcParameterDefinitionName,	DstParameter);
 				DstParameter->BindToTypeInfo();
-
-				//DstParameter->Name = FWebAPINameInfo(NameTransformer(DstParameter->Name), DstParameter->Name.GetJsonName(), DstParameter->Type);
-				DstParameter->Description = InSrcParameter->Description.Get(TEXT(""));
-				DstParameter->bIsRequired = InSrcParameter->bRequired.Get(false);
-				
-				/*
-				if(!DstParameter->Type.TypeInfo->bIsBuiltinType)
-				{
-					if(DstParameter->Type.TypeInfo->IsEnum())
-					{
-						ConvertEnum(InSrcParameter);						
-					}
-					else
-					{
-						unimplemented();						
-					}
-				}
-				*/
-
-				/**
+			
 				// Special case - if the name is "body", set as Body property and not Parameter
-				if (InSrcParameter->Name.Equals(TEXT("body"), ESearchCase::IgnoreCase) && !SrcParameterDefinitionName.IsEmpty())
-				{
+				if (SrcParameter->Name.Equals(TEXT("body"), ESearchCase::IgnoreCase) && !SrcParameterDefinitionName.IsEmpty())
+				{ 
 					// Check for existing definition
 					if (const TObjectPtr<UWebAPITypeInfo>* FoundGeneratedType = OutputSchema->TypeRegistry->FindGeneratedType(EWebAPISchemaType::Model, SrcParameterDefinitionName))
 					{
 						DstParameter->Model = Cast<UWebAPIModel>((*FoundGeneratedType)->Model.LoadSynchronous());
-						return DstParameter;
+						return true;
 					}
 				}
-				*/
-
-				return DstParameter;
 			}
-
-			bool FWebAPIOpenAPISchemaConverter::ConvertRequest(const FWebAPITypeNameVariant& InOperationName, const TSharedPtr<OpenAPI::V3::FOperationObject>& InOperation, const TObjectPtr<UWebAPIOperationRequest>& OutRequest)
+		}
+		else if(TSharedPtr<OpenAPI::V3::FRequestBodyObject> SrcRequestBody = ResolveReference(InOperation->RequestBody))
+		{
+			if(!SrcRequestBody->Content.IsEmpty())
 			{
-				// const FWebAPITypeNameVariant RequestTypeName = OutputSchema->TypeRegistry->AddUnnamedType(IWebAPIEditorModuleInterface::Get().GetStaticTypeRegistry()->Object);
-				const FWebAPITypeNameVariant RequestTypeName = OutputSchema->TypeRegistry->GetOrMakeGeneratedType(
-					EWebAPISchemaType::Model,
-					ProviderSettings.MakeRequestTypeName(InOperationName),
-					InOperationName.ToString(true),
-					IWebAPIEditorModuleInterface::Get().GetStaticTypeRegistry()->Object);
-#if WITH_WEBAPI_DEBUG
-				RequestTypeName.TypeInfo->DebugString += TEXT(">ConvertRequest");
-#endif
-				RequestTypeName.TypeInfo->Prefix = TEXT("F");
-
-				OutRequest->Name = RequestTypeName;
-
-				if(!InOperation->Parameters.IsEmpty())
+				TSharedPtr<V3::FMediaTypeObject> MediaType = nullptr;
+			
+				// Choose json entry or whatever is first (if json not found)
+				if(const TSharedPtr<V3::FMediaTypeObject>* JsonMediaType = SrcRequestBody->Content.Find(UE::WebAPI::MimeType::NAME_Json.ToString()))
 				{
-					for (Json::TJsonReference<OpenAPI::V3::FParameterObject>& SrcParameter : InOperation->Parameters)
-					{
-						FString SrcParameterDefinitionName;
-						// Will get schema or create if it doesn't exist (but will be empty)
-						TSharedPtr<OpenAPI::V3::FSchemaObject> SrcParameterSchema = ResolveReference(SrcParameter->Schema.Get({}), SrcParameterDefinitionName);
-
-						FWebAPINameVariant ParameterName = FWebAPINameInfo(NameTransformer(SrcParameter->Name), SrcParameter->Name);
-						if(SrcParameterDefinitionName.IsEmpty())
-						{
-							SrcParameterDefinitionName = ProviderSettings.MakeNestedPropertyTypeName(InOperationName, ParameterName);
-						}
-
-						const TObjectPtr<UWebAPIOperationParameter> DstParameter = OutRequest->Parameters.Add_GetRef(NewObject<UWebAPIOperationParameter>(OutRequest));
-						ConvertOperationParameter(ParameterName, SrcParameter.GetShared(), SrcParameterDefinitionName,	DstParameter);
-						DstParameter->BindToTypeInfo();
-					
-						// Special case - if the name is "body", set as Body property and not Parameter
-						if (SrcParameter->Name.Equals(TEXT("body"), ESearchCase::IgnoreCase) && !SrcParameterDefinitionName.IsEmpty())
-						{ 
-							// Check for existing definition
-							if (const TObjectPtr<UWebAPITypeInfo>* FoundGeneratedType = OutputSchema->TypeRegistry->FindGeneratedType(EWebAPISchemaType::Model, SrcParameterDefinitionName))
-							{
-								DstParameter->Model = Cast<UWebAPIModel>((*FoundGeneratedType)->Model.LoadSynchronous());
-								return true;
-							}
-						}
-					}
+					MediaType = *JsonMediaType;
 				}
-				else if(TSharedPtr<OpenAPI::V3::FRequestBodyObject> SrcRequestBody = ResolveReference(InOperation->RequestBody))
+				else if(!SrcRequestBody->Content.IsEmpty())
 				{
-					if(!SrcRequestBody->Content.IsEmpty())
-					{
-						TSharedPtr<V3::FMediaTypeObject> MediaType = nullptr;
-					
-						// Choose json entry or whatever is first (if json not found)
-						if(const TSharedPtr<V3::FMediaTypeObject>* JsonMediaType = SrcRequestBody->Content.Find(UE::WebAPI::MimeType::NAME_Json.ToString()))
-						{
-							MediaType = *JsonMediaType;
-						}
-						else if(!SrcRequestBody->Content.IsEmpty())
-						{
-							TArray<FString> Keys;
-							SrcRequestBody->Content.GetKeys(Keys);
-							MediaType = SrcRequestBody->Content[Keys[0]];
-						}
-
-						if(MediaType.IsValid())
-						{
-							FString MediaTypeDefinitionName;
-							TSharedPtr<OpenAPI::V3::FSchemaObject> MediaTypeSchema = ResolveReference(MediaType->Schema, MediaTypeDefinitionName);
-
-							// @todo: what if it's an array of $ref? TypeInfo doesn't specify whether it's an array or not - check spec
-							const TObjectPtr<UWebAPIOperationParameter> DstParameter = OutRequest->Parameters.Add_GetRef(NewObject<UWebAPIOperationParameter>(OutRequest));
-							DstParameter->Type = ResolveType(MediaTypeSchema, MediaTypeDefinitionName);
-							check(DstParameter->Type.IsValid());
-							
-							DstParameter->Name = ProviderSettings.GetDefaultPropertyName();
-							DstParameter->Name = ResolvePropertyName(DstParameter, DstParameter->Type.ToString(true));
-							check(DstParameter->Name.IsValid());
-
-							DstParameter->Storage = EWebAPIParameterStorage::Body;
-							DstParameter->BindToTypeInfo();
-						}
-					}
+					TArray<FString> Keys;
+					SrcRequestBody->Content.GetKeys(Keys);
+					MediaType = SrcRequestBody->Content[Keys[0]];
 				}
-				// Operation has no parameters, but if arbitrary json is enabled, make a Value property of type JsonObject
-				else if(ProviderSettings.bEnableArbitraryJsonPayloads)
+
+				if(MediaType.IsValid())
 				{
-					FWebAPINameVariant ParameterName = TEXT("Value");
+					FString MediaTypeDefinitionName;
+					TSharedPtr<OpenAPI::V3::FSchemaObject> MediaTypeSchema = ResolveReference(MediaType->Schema, MediaTypeDefinitionName);
 
+					// @todo: what if it's an array of $ref? TypeInfo doesn't specify whether it's an array or not - check spec
 					const TObjectPtr<UWebAPIOperationParameter> DstParameter = OutRequest->Parameters.Add_GetRef(NewObject<UWebAPIOperationParameter>(OutRequest));
-					DstParameter->Type = IWebAPIEditorModuleInterface::Get().GetStaticTypeRegistry()->JsonObject;
-					DstParameter->Name = FWebAPINameInfo(ParameterName.ToString(true), ParameterName.GetJsonName(), DstParameter->Type);
-					DstParameter->bIsMixin = true; // treat as object, not as field
+					DstParameter->Type = ResolveType(MediaTypeSchema, MediaTypeDefinitionName);
+					check(DstParameter->Type.IsValid());
+					
+					DstParameter->Name = ProviderSettings.GetDefaultPropertyName();
+					DstParameter->Name = ResolvePropertyName(DstParameter, DstParameter->Type.ToString(true));
+					check(DstParameter->Name.IsValid());
+
 					DstParameter->Storage = EWebAPIParameterStorage::Body;
 					DstParameter->BindToTypeInfo();
 				}
-
-				OutRequest->BindToTypeInfo();
-				
-				return true;
 			}
+		}
+		// Operation has no parameters, but if arbitrary json is enabled, make a Value property of type JsonObject
+		else if(ProviderSettings.bEnableArbitraryJsonPayloads)
+		{
+			FWebAPINameVariant ParameterName = TEXT("Value");
+
+			const TObjectPtr<UWebAPIOperationParameter> DstParameter = OutRequest->Parameters.Add_GetRef(NewObject<UWebAPIOperationParameter>(OutRequest));
+			DstParameter->Type = IWebAPIEditorModuleInterface::Get().GetStaticTypeRegistry()->JsonObject;
+			DstParameter->Name = FWebAPINameInfo(ParameterName.ToString(true), ParameterName.GetJsonName(), DstParameter->Type);
+			DstParameter->bIsMixin = true; // treat as object, not as field
+			DstParameter->Storage = EWebAPIParameterStorage::Body;
+			DstParameter->BindToTypeInfo();
+		}
+
+		OutRequest->BindToTypeInfo();
+		
+		return true;
+	}
 
-			bool FWebAPIOpenAPISchemaConverter::ConvertResponse(const FWebAPITypeNameVariant& InOperationName, uint32 InResponseCode, const TSharedPtr<OpenAPI::V3::FResponseObject>& InResponse, const TObjectPtr<UWebAPIOperationResponse>& OutResponse)
-			{
-				const FWebAPITypeNameVariant ResponseTypeName = OutputSchema->TypeRegistry->GetOrMakeGeneratedType(
-					EWebAPISchemaType::Model,
-					ProviderSettings.MakeResponseTypeName(InOperationName, InResponseCode),
-					InOperationName.ToString(true),
-					IWebAPIEditorModuleInterface::Get().GetStaticTypeRegistry()->Object);
+	bool FWebAPIOpenAPISchemaConverter::ConvertResponse(const FWebAPITypeNameVariant& InOperationName, uint32 InResponseCode, const TSharedPtr<OpenAPI::V3::FResponseObject>& InResponse, const TObjectPtr<UWebAPIOperationResponse>& OutResponse)
+	{
+		const FWebAPITypeNameVariant ResponseTypeName = OutputSchema->TypeRegistry->GetOrMakeGeneratedType(
+			EWebAPISchemaType::Model,
+			ProviderSettings.MakeResponseTypeName(InOperationName, InResponseCode),
+			InOperationName.ToString(true),
+			IWebAPIEditorModuleInterface::Get().GetStaticTypeRegistry()->Object);
 #if WITH_WEBAPI_DEBUG
-				ResponseTypeName.TypeInfo->DebugString += TEXT(">ConvertResponse");
+		ResponseTypeName.TypeInfo->DebugString += TEXT(">ConvertResponse");
 #endif
-				ResponseTypeName.TypeInfo->Prefix = TEXT("F");
+		ResponseTypeName.TypeInfo->Prefix = TEXT("F");
 
-				check(ResponseTypeName.HasTypeInfo());
+		check(ResponseTypeName.HasTypeInfo());
 
-				OutResponse->Name = ResponseTypeName;
-				OutResponse->Code = InResponseCode;
-				OutResponse->Description = InResponse->Description;
-				OutResponse->Message = OutResponse->Description;
+		OutResponse->Name = ResponseTypeName;
+		OutResponse->Code = InResponseCode;
+		OutResponse->Description = InResponse->Description;
+		OutResponse->Message = OutResponse->Description;
 
-				if (InResponse->Content.IsSet())
-				{
-					for(const TPair<FString, Json::TJsonReference<V3::FMediaTypeObject>>& Item : InResponse->Content.GetValue())
-					{
-						// Will get schema or create if it doesn't exist (but will be empty)
-						FString SrcPropertyDefinitionName;
-						TSharedPtr<V3::FSchemaObject> SrcResponseSchema = ResolveReference(Item.Value->Schema, SrcPropertyDefinitionName);
-
-						const FString PropertyName = SrcResponseSchema->Type.Get(GetDefaultJsonTypeForStorage(OutResponse->Storage)) == TEXT("array") ? ProviderSettings.GetDefaultArrayPropertyName() : ProviderSettings.GetDefaultPropertyName();
-						const TObjectPtr<UWebAPIProperty>& DstProperty = OutResponse->Properties.Add_GetRef(NewObject<UWebAPIProperty>(OutResponse));
-						DstProperty->bIsMixin = true;
-						ConvertProperty(OutResponse->Name,
-							PropertyName,
-							SrcResponseSchema,
-							SrcPropertyDefinitionName,
-							DstProperty);
-
-						/*
-						//OutResponse->bIsArray = IsArray(InSrcParameter);
-						if (!DstProperty->Type.HasTypeInfo())
-						{
-							DstProperty->Type = ResolveType(SrcResponseSchema->Type.Get(TEXT("object")),
-								SrcResponseSchema->Format.Get(TEXT("")),
-								SrcPropertyDefinitionName,
-								SrcResponseSchema.GetShared());
-						}
+		if (InResponse->Content.IsSet())
+		{
+			for(const TPair<FString, Json::TJsonReference<V3::FMediaTypeObject>>& Item : InResponse->Content.GetValue())
+			{
+				// Will get schema or create if it doesn't exist (but will be empty)
+				FString SrcPropertyDefinitionName;
+				TSharedPtr<V3::FSchemaObject> SrcResponseSchema = ResolveReference(Item.Value->Schema, SrcPropertyDefinitionName);
+
+				const FString PropertyName = SrcResponseSchema->Type.Get(GetDefaultJsonTypeForStorage(OutResponse->Storage)) == TEXT("array") ? ProviderSettings.GetDefaultArrayPropertyName() : ProviderSettings.GetDefaultPropertyName();
+				const TObjectPtr<UWebAPIProperty>& DstProperty = OutResponse->Properties.Add_GetRef(NewObject<UWebAPIProperty>(OutResponse));
+				DstProperty->bIsMixin = true;
+				ConvertProperty(OutResponse->Name,
+					PropertyName,
+					SrcResponseSchema,
+					SrcPropertyDefinitionName,
+					DstProperty);
 
-						
-						if(SrcPropertyDefinitionName.IsEmpty())
-						{
-							SrcPropertyDefinitionName = ProviderSettings.MakeNestedPropertyTypeName(InOperationName, PropertyName);
-						}
+				/*
+				//OutResponse->bIsArray = IsArray(InSrcParameter);
+				if (!DstProperty->Type.HasTypeInfo())
+				{
+					DstProperty->Type = ResolveType(SrcResponseSchema->Type.Get(TEXT("object")),
+						SrcResponseSchema->Format.Get(TEXT("")),
+						SrcPropertyDefinitionName,
+						SrcResponseSchema.GetShared());
+				}
 
-						const FString Type = InResponse->Schema->Get()->Type.Get(TEXT("object"));
+				
+				if(SrcPropertyDefinitionName.IsEmpty())
+				{
+					SrcPropertyDefinitionName = ProviderSettings.MakeNestedPropertyTypeName(InOperationName, PropertyName);
+				}
 
-						const TObjectPtr<UWebAPIModel> Model = OutResponse;
-						if (!ConvertModel(SrcResponseSchema.GetShared(), ResponseTypeName, Model))
-						{
-							return false;
-						}
+				const FString Type = InResponse->Schema->Get()->Type.Get(TEXT("object"));
 
-						// Don't use ParameterTypeName->Name, it might have a Parameter specific pre/postfix.
-						FString SrcResponseDefinitionName = ResponseTypeName.ToString(true);
+				const TObjectPtr<UWebAPIModel> Model = OutResponse;
+				if (!ConvertModel(SrcResponseSchema.GetShared(), ResponseTypeName, Model))
+				{
+					return false;
+				}
 
-						// Will get schema or create if it doesn't exist (but will be empty)
-						const TSharedPtr<OpenAPI::V3::FSchemaObject> SrcParameterSchema = ResolveReference(InResponse->Schema.Get({}), SrcResponseDefinitionName, false);
+				// Don't use ParameterTypeName->Name, it might have a Parameter specific pre/postfix.
+				FString SrcResponseDefinitionName = ResponseTypeName.ToString(true);
 
-						const FString PropertyName = Type == TEXT("array") ? ProviderSettings.GetDefaultArrayPropertyName() : ProviderSettings.GetDefaultPropertyName();
-						// FString DefinitionName = ProviderSettings.MakeNestedPropertyTypeName(DefinitionName, PropertyName);
-						
-						const TObjectPtr<UWebAPIProperty>& DstProperty = OutResponse->Properties.Add_GetRef(NewObject<UWebAPIProperty>(OutResponse));
-						ConvertProperty(OutResponse->Name,
-							PropertyName,
-							SrcParameterSchema,
-							SrcResponseDefinitionName,
-							DstProperty);
-						
+				// Will get schema or create if it doesn't exist (but will be empty)
+				const TSharedPtr<OpenAPI::V3::FSchemaObject> SrcParameterSchema = ResolveReference(InResponse->Schema.Get({}), SrcResponseDefinitionName, false);
 
-						if (Type == TEXT("array"))
-						{
-	 
-						}
-						*/
-						// @todo: Account for single item = single property, ie. type=array, items=SomeDef would be TArray<FSomeDef> Items 
-					}
-				}
+				const FString PropertyName = Type == TEXT("array") ? ProviderSettings.GetDefaultArrayPropertyName() : ProviderSettings.GetDefaultPropertyName();
+				// FString DefinitionName = ProviderSettings.MakeNestedPropertyTypeName(DefinitionName, PropertyName);
+				
+				const TObjectPtr<UWebAPIProperty>& DstProperty = OutResponse->Properties.Add_GetRef(NewObject<UWebAPIProperty>(OutResponse));
+				ConvertProperty(OutResponse->Name,
+					PropertyName,
+					SrcParameterSchema,
+					SrcResponseDefinitionName,
+					DstProperty);
+				
 
-				if (InResponse->Headers.IsSet())
+				if (Type == TEXT("array"))
 				{
-					for (const TPair<FString, Json::TJsonReference<OpenAPI::V3::FHeaderObject>>& SrcHeader : InResponse->Headers.GetValue())
-					{
-						FString Key = SrcHeader.Key; // @todo: what is this?
-						// SrcHeader.Value-> // @todo: this describes return object
-					}
-				}
-
-				OutResponse->BindToTypeInfo();
 
-				return true;
+				}
+				*/
+				// @todo: Account for single item = single property, ie. type=array, items=SomeDef would be TArray<FSomeDef> Items 
 			}
+		}
 
-			TObjectPtr<UWebAPIOperation> FWebAPIOpenAPISchemaConverter::ConvertOperation(const FString& InPath, const FString& InVerb, const TSharedPtr<OpenAPI::V3::FOperationObject>& InSrcOperation, const FWebAPITypeNameVariant& InOperationTypeName)
+		if (InResponse->Headers.IsSet())
+		{
+			for (const TPair<FString, Json::TJsonReference<OpenAPI::V3::FHeaderObject>>& SrcHeader : InResponse->Headers.GetValue())
 			{
-				
-				FString OperationName = InSrcOperation->OperationId.Get(InSrcOperation->Summary.Get(TEXT("")));
-				if(OperationName.IsEmpty())
-				{
-					// ie. GET /pets/{id} == GetPetsById
-					
-					TArray<FString> SplitPath;
-					InPath.ParseIntoArray(SplitPath, TEXT("/"));
+				FString Key = SrcHeader.Key; // @todo: what is this?
+				// SrcHeader.Value-> // @todo: this describes return object
+			}
+		}
 
-					Algo::ForEach(SplitPath, [&](FString& Str)
-					{
-						Str = NameTransformer(Str);
-					});
+		OutResponse->BindToTypeInfo();
 
-					OperationName = NameTransformer(InVerb) + FString::Join(SplitPath, TEXT(""));
-					
-				}
-				
-				OperationName = NameTransformer(OperationName);
-				OperationName = ProviderSettings.MakeValidMemberName(OperationName, InVerb);
+		return true;
+	}
 
-				check(!OperationName.IsEmpty());
-				
-				FWebAPITypeNameVariant OperationTypeName;
-				if(InOperationTypeName.IsValid())
-				{
-					OperationName = InOperationTypeName.ToString(true);
-					OperationTypeName = InOperationTypeName;
-				}
-				else
-				{
-					OperationTypeName = OutputSchema->TypeRegistry->GetOrMakeGeneratedType(
-						EWebAPISchemaType::Operation,
-						NameTransformer(*OperationName),
-						*OperationName,
-						TEXT("U"));
-					OperationTypeName.TypeInfo->SetName(OperationName);
+	TObjectPtr<UWebAPIOperation> FWebAPIOpenAPISchemaConverter::ConvertOperation(const FString& InPath, const FString& InVerb, const TSharedPtr<OpenAPI::V3::FOperationObject>& InSrcOperation, const FWebAPITypeNameVariant& InOperationTypeName)
+	{
+		
+		FString OperationName = InSrcOperation->OperationId.Get(InSrcOperation->Summary.Get(TEXT("")));
+		if(OperationName.IsEmpty())
+		{
+			// ie. GET /pets/{id} == GetPetsById
+			
+			TArray<FString> SplitPath;
+			InPath.ParseIntoArray(SplitPath, TEXT("/"));
+
+			Algo::ForEach(SplitPath, [&](FString& Str)
+			{
+				Str = NameTransformer(Str);
+			});
+
+			OperationName = NameTransformer(InVerb) + FString::Join(SplitPath, TEXT(""));
+			
+		}
+		
+		OperationName = NameTransformer(OperationName);
+		OperationName = ProviderSettings.MakeValidMemberName(OperationName, InVerb);
+
+		check(!OperationName.IsEmpty());
+		
+		FWebAPITypeNameVariant OperationTypeName;
+		if(InOperationTypeName.IsValid())
+		{
+			OperationName = InOperationTypeName.ToString(true);
+			OperationTypeName = InOperationTypeName;
+		}
+		else
+		{
+			OperationTypeName = OutputSchema->TypeRegistry->GetOrMakeGeneratedType(
+				EWebAPISchemaType::Operation,
+				NameTransformer(*OperationName),
+				*OperationName,
+				TEXT("U"));
+			OperationTypeName.TypeInfo->SetName(OperationName);
 
 #if WITH_WEBAPI_DEBUG
-					OperationTypeName.TypeInfo->DebugString += TEXT(">ConvertOperation");
+			OperationTypeName.TypeInfo->DebugString += TEXT(">ConvertOperation");
 #endif
+		}
+
+		// A spec can have no tags, so ensure there's a default
+		const FString FirstTag = ProviderSettings.ToPascalCase(InSrcOperation->Tags.Get({TEXT("Default")})[0]);
+		TObjectPtr<UWebAPIService>* Service = OutputSchema->Services.Find(FirstTag);
+		checkf(Service, TEXT("An operation must belong to a service!"));
+
+		const TObjectPtr<UWebAPIOperation> DstOperation = (*Service)->Operations.Add_GetRef(NewObject<UWebAPIOperation>(*Service));
+		DstOperation->Service = *Service;
+		DstOperation->Verb = InVerb;
+		DstOperation->Path = InPath;
+
+		DstOperation->Name = OperationTypeName;
+		// Choose first non-empty of: description, summary, operation id
+		DstOperation->Description = InSrcOperation->Description.Get(
+			InSrcOperation->OperationId.Get(InSrcOperation->Summary.Get(
+				InSrcOperation->OperationId.Get(TEXT("")))));
+		DstOperation->bIsDeprecated = InSrcOperation->bDeprecated.Get(false);
+
+		ConvertRequest(DstOperation->Name, InSrcOperation, DstOperation->Request);
+
+		// @todo:
+		/*
+		if (InSrcOperation->Consumes.IsSet())
+		{
+			DstOperation->RequestContentTypes.Append(InSrcOperation->Consumes.GetValue());
+		}
+
+		if (InSrcOperation->Produces.IsSet())
+		{
+			DstOperation->ResponseContentTypes.Append(InSrcOperation->Produces.GetValue());
+		}
+		*/
+
+		if (ensure(!InSrcOperation->Responses.IsEmpty()))
+		{
+			for (const TPair<FString, Json::TJsonReference<OpenAPI::V3::FResponseObject>>& SrcResponse : InSrcOperation->Responses)
+			{
+				// If "Default", resolves to 0, and means all other unhandled codes (similar to default in switch statement)
+				const uint32 Code = FCString::Atoi(*SrcResponse.Key);
+				const TObjectPtr<UWebAPIOperationResponse> DstResponse = DstOperation->Responses.Add_GetRef(NewObject<UWebAPIOperationResponse>(DstOperation));
+				// OutOperation->AddResponse() @todo
+				ConvertResponse(DstOperation->Name, Code, ResolveReference(SrcResponse.Value), DstResponse);
+
+				// If success response (code 200), had no resolved properties but the operation says it returns something, then add that something as a property
+				if(DstResponse->Properties.IsEmpty() && !DstOperation->ResponseContentTypes.IsEmpty() && DstResponse->Code == 200)
+				{
+					const FString PropertyName = ProviderSettings.GetDefaultPropertyName();
+					const TObjectPtr<UWebAPIProperty>& DstProperty = DstResponse->Properties.Add_GetRef(NewObject<UWebAPIProperty>(DstResponse));
+					DstProperty->bIsMixin = true;
+					ConvertProperty(DstResponse->Name,
+						PropertyName,
+						GetTypeForContentType(DstOperation->ResponseContentTypes[0]),
+						DstProperty);
+					DstProperty->Name = PropertyName;
+					DstResponse->Storage = EWebAPIResponseStorage::Body;
 				}
+			}
+		}
 
-				// A spec can have no tags, so ensure there's a default
-				const FString FirstTag = ProviderSettings.ToPascalCase(InSrcOperation->Tags.Get({TEXT("Default")})[0]);
-				TObjectPtr<UWebAPIService>* Service = OutputSchema->Services.Find(FirstTag);
-				checkf(Service, TEXT("An operation must belong to a service!"));
+		check(DstOperation->Name.HasTypeInfo());
 
-				const TObjectPtr<UWebAPIOperation> DstOperation = (*Service)->Operations.Add_GetRef(NewObject<UWebAPIOperation>(*Service));
-				DstOperation->Service = *Service;
-				DstOperation->Verb = InVerb;
-				DstOperation->Path = InPath;
+		const FName OperationObjectName = ProviderSettings.MakeOperationObjectName(*Service, OperationName);
+		DstOperation->Rename(*OperationObjectName.ToString(), DstOperation->GetOuter());
+		
+		return DstOperation;
+	}
 
-				DstOperation->Name = OperationTypeName;
-				// Choose first non-empty of: description, summary, operation id
-				DstOperation->Description = InSrcOperation->Description.Get(
-					InSrcOperation->OperationId.Get(InSrcOperation->Summary.Get(
-						InSrcOperation->OperationId.Get(TEXT("")))));
-				DstOperation->bIsDeprecated = InSrcOperation->bDeprecated.Get(false);
+	TObjectPtr<UWebAPIService> FWebAPIOpenAPISchemaConverter::ConvertService(const FWebAPINameVariant& InName) const
+	{
+		const FString TagName = NameTransformer(InName);
+		return OutputSchema->GetOrMakeService(TagName);
+	}
 
-				ConvertRequest(DstOperation->Name, InSrcOperation, DstOperation->Request);
+	TObjectPtr<UWebAPIService> FWebAPIOpenAPISchemaConverter::ConvertService(const TSharedPtr<OpenAPI::V3::FTagObject>& InTag) const
+	{
+		const TObjectPtr<UWebAPIService> Service = ConvertService(InTag->Name);
+		Service->Description = InTag->Description.Get(TEXT(""));
+		return Service;
+	}
 
-				// @todo:
-				/*
-				if (InSrcOperation->Consumes.IsSet())
-				{
-					DstOperation->RequestContentTypes.Append(InSrcOperation->Consumes.GetValue());
-				}
+	bool FWebAPIOpenAPISchemaConverter::ConvertModels(const TMap<FString, Json::TJsonReference<OpenAPI::V3::FSchemaObject>>& InSchemas,	UWebAPISchema* OutSchema)
+	{
+		bool bAllConverted = true;
+		for (const TTuple<FString, Json::TJsonReference<OpenAPI::V3::FSchemaObject>>& Item : InSchemas)
+		{
+			const FString& Name = Item.Key;
+			const TSharedPtr<V3::FSchemaObject> Schema = ResolveReference(Item.Value);
 
-				if (InSrcOperation->Produces.IsSet())
-				{
-					DstOperation->ResponseContentTypes.Append(InSrcOperation->Produces.GetValue());
-				}
-				*/
+			if (!Schema.IsValid())
+			{
+				bAllConverted = false;
+				FFormatNamedArguments Args;
+				Args.Add(TEXT("ModelName"), FText::FromString(Name));
+				MessageLog->LogWarning(
+					FText::Format(LOCTEXT("SchemaInvalid", "The schema for model \"{ModelName}\" was invalid/null."), Args),
+					FWebAPIOpenAPIProvider::LogName);
+				continue;
+			}
 
-				if (ensure(!InSrcOperation->Responses.IsEmpty()))
-				{
-					for (const TPair<FString, Json::TJsonReference<OpenAPI::V3::FResponseObject>>& SrcResponse : InSrcOperation->Responses)
-					{
-						// If "Default", resolves to 0, and means all other unhandled codes (similar to default in switch statement)
-						const uint32 Code = FCString::Atoi(*SrcResponse.Key);
-						const TObjectPtr<UWebAPIOperationResponse> DstResponse = DstOperation->Responses.Add_GetRef(NewObject<UWebAPIOperationResponse>(DstOperation));
-						// OutOperation->AddResponse() @todo
-						ConvertResponse(DstOperation->Name, Code, ResolveReference(SrcResponse.Value), DstResponse);
-
-						// If success response (code 200), had no resolved properties but the operation says it returns something, then add that something as a property
-						if(DstResponse->Properties.IsEmpty() && !DstOperation->ResponseContentTypes.IsEmpty() && DstResponse->Code == 200)
-						{
-							const FString PropertyName = ProviderSettings.GetDefaultPropertyName();
-							const TObjectPtr<UWebAPIProperty>& DstProperty = DstResponse->Properties.Add_GetRef(NewObject<UWebAPIProperty>(DstResponse));
-							DstProperty->bIsMixin = true;
-							ConvertProperty(DstResponse->Name,
-								PropertyName,
-								GetTypeForContentType(DstOperation->ResponseContentTypes[0]),
-								DstProperty);
-							DstProperty->Name = PropertyName;
-							DstResponse->Storage = EWebAPIResponseStorage::Body;
-						}
-					}
-				}
+			FWebAPITypeNameVariant ModelName = OutputSchema->TypeRegistry->GetOrMakeGeneratedType(
+				EWebAPISchemaType::Model,
+				NameTransformer(Name),
+				Name,
+				TEXT("F"));
 
-				check(DstOperation->Name.HasTypeInfo());
+			const TObjectPtr<UWebAPIModel>& Model = OutSchema->AddModel<UWebAPIModel>(ModelName.TypeInfo.Get());
+			ConvertModel(Schema, {}, Model);
+		}
 
-				const FName OperationObjectName = ProviderSettings.MakeOperationObjectName(*Service, OperationName);
-				DstOperation->Rename(*OperationObjectName.ToString(), DstOperation->GetOuter());
-				
-				return DstOperation;
+		// Set Model property of TypeInfos where applicable
+		for (const TObjectPtr<UWebAPIModelBase>& ModelBase : OutputSchema->Models)
+		{
+			if (const TObjectPtr<UWebAPIModel>& Model = Cast<UWebAPIModel>(ModelBase))
+			{
+				Model->BindToTypeInfo();
+				// Model->Name.TypeInfo->Model = Model;
 			}
-
-			TObjectPtr<UWebAPIService> FWebAPIOpenAPISchemaConverter::ConvertService(const FWebAPINameVariant& InName) const
+			else if (const TObjectPtr<UWebAPIEnum>& Enum = Cast<UWebAPIEnum>(ModelBase))
 			{
-				const FString TagName = NameTransformer(InName);
-				return OutputSchema->GetOrMakeService(TagName);
+				Enum->BindToTypeInfo();
+				// Enum->Name.TypeInfo->Model = Enum;
 			}
-
-			TObjectPtr<UWebAPIService> FWebAPIOpenAPISchemaConverter::ConvertService(const TSharedPtr<OpenAPI::V3::FTagObject>& InTag) const
+			else if (const TObjectPtr<UWebAPIParameter>& ParameterModel = Cast<UWebAPIParameter>(ModelBase))
 			{
-				const TObjectPtr<UWebAPIService> Service = ConvertService(InTag->Name);
-				Service->Description = InTag->Description.Get(TEXT(""));
-				return Service;
+				ParameterModel->BindToTypeInfo();
+				// ParameterModel->Type.TypeInfo->Model = ParameterModel;
 			}
+		}
 
-			bool FWebAPIOpenAPISchemaConverter::ConvertModels(const TMap<FString, Json::TJsonReference<OpenAPI::V3::FSchemaObject>>& InSchemas,	UWebAPISchema* OutSchema)
-			{
-				bool bAllConverted = true;
-				for (const TTuple<FString, Json::TJsonReference<OpenAPI::V3::FSchemaObject>>& Item : InSchemas)
-				{
-					const FString& Name = Item.Key;
-					const TSharedPtr<V3::FSchemaObject> Schema = ResolveReference(Item.Value);
+		return bAllConverted;
+	}
 
-					if (!Schema.IsValid())
-					{
-						bAllConverted = false;
-						FFormatNamedArguments Args;
-						Args.Add(TEXT("ModelName"), FText::FromString(Name));
-						MessageLog->LogWarning(
-							FText::Format(LOCTEXT("SchemaInvalid", "The schema for model \"{ModelName}\" was invalid/null."), Args),
-							FWebAPIOpenAPIProvider::LogName);
-						continue;
-					}
+	bool FWebAPIOpenAPISchemaConverter::ConvertParameters(const TArray<Json::TJsonReference<OpenAPI::V3::FParameterObject>>& InParameters, UWebAPISchema* OutSchema)
+	{
+		bool bAllConverted = true;
+		for (const Json::TJsonReference<OpenAPI::V3::FParameterObject>& Item : InParameters)
+		{
+			bAllConverted &= ConvertParameter(ResolveReference(Item)) != nullptr;
+		}
 
-					FWebAPITypeNameVariant ModelName = OutputSchema->TypeRegistry->GetOrMakeGeneratedType(
-						EWebAPISchemaType::Model,
-						NameTransformer(Name),
-						Name,
-						TEXT("F"));
+		// Re-bind models to their TypeInfos
+		for (const TObjectPtr<UWebAPIModelBase>& Model : OutputSchema->Models)
+		{
+			Model->BindToTypeInfo();
+		}
 
-					const TObjectPtr<UWebAPIModel>& Model = OutSchema->AddModel<UWebAPIModel>(ModelName.TypeInfo.Get());
-					ConvertModel(Schema, {}, Model);
-				}
+		return bAllConverted;
+	}
+
+	bool FWebAPIOpenAPISchemaConverter::ConvertParameters(const TMap<FString, TSharedPtr<OpenAPI::V3::FParameterObject>>& InParameters, UWebAPISchema* OutSchema)
+	{
+		bool bAllConverted = true;
+		for (const TTuple<FString, TSharedPtr<OpenAPI::V3::FParameterObject>>& Item : InParameters)
+		{
+			const FWebAPITypeNameVariant& Name = Item.Key;
+			const TSharedPtr<OpenAPI::V3::FParameterObject>& SrcParameter = Item.Value;
+
+			bAllConverted &= ConvertParameter(SrcParameter) != nullptr;
+		}
+
+		// Re-bind models to their TypeInfos
+		for (const TObjectPtr<UWebAPIModelBase>& Model : OutputSchema->Models)
+		{
+			Model->BindToTypeInfo();
+		}
+
+		return bAllConverted;
+	}
 
-				// Set Model property of TypeInfos where applicable
-				for (const TObjectPtr<UWebAPIModelBase>& ModelBase : OutputSchema->Models)
+	bool FWebAPIOpenAPISchemaConverter::ConvertSecurity(const UE::WebAPI::OpenAPI::V3::FOpenAPIObject& InOpenAPI, UWebAPISchema* OutSchema)
+	{
+		if (!InOpenAPI.Components->SecuritySchemes.IsEmpty())
+		{
+			for (const TPair<FString, Json::TJsonReference<OpenAPI::V3::FSecuritySchemeObject>>& Item : InOpenAPI.Components->SecuritySchemes)
+			{
+				const FString& Name = Item.Key;
+				const TSharedPtr<OpenAPI::V3::FSecuritySchemeObject>& SecurityScheme = ResolveReference(Item.Value);
+				if(SecurityScheme->Type == TEXT("http") && SecurityScheme->Scheme == TEXT("basic"))
+				{
+					// type = basic
+				}
+				else if(SecurityScheme->Type == TEXT("http") && SecurityScheme->Scheme == TEXT("bearer"))
 				{
-					if (const TObjectPtr<UWebAPIModel>& Model = Cast<UWebAPIModel>(ModelBase))
+					// type = apiKey
+					// name = authorization
+					// storage = header
+				}
+				else if(SecurityScheme->Type == TEXT("oauth2"))
+				{
+					// @todo:
+					FString FlowName = SecurityScheme->Flows.IsValid() ? TEXT("") : TEXT("");
+					V3::FOAuthFlowsObject Flow = (SecurityScheme->Flows.Get())[0];
+
+					if(FlowName == TEXT("clientCredentials"))
 					{
-						Model->BindToTypeInfo();
-						// Model->Name.TypeInfo->Model = Model;
+						// flow = application
 					}
-					else if (const TObjectPtr<UWebAPIEnum>& Enum = Cast<UWebAPIEnum>(ModelBase))
+					else if(FlowName == TEXT("authorizationCode"))
 					{
-						Enum->BindToTypeInfo();
-						// Enum->Name.TypeInfo->Model = Enum;
+						// flow = accessCode
 					}
-					else if (const TObjectPtr<UWebAPIParameter>& ParameterModel = Cast<UWebAPIParameter>(ModelBase))
+					else
 					{
-						ParameterModel->BindToTypeInfo();
-						// ParameterModel->Type.TypeInfo->Model = ParameterModel;
+						// flow = FlowName
 					}
+
+					
 				}
 
-				return bAllConverted;
+				// @todo: WebAPI security primitive
+				// TObjectPtr<UWebAPIModel> Model = OutSchema->Client->Models.Add_GetRef(NewObject<UWebAPIModel>(OutSchema));
+				// ConvertModel(Name, Schema.Get(), Model);
 			}
 
-			bool FWebAPIOpenAPISchemaConverter::ConvertParameters(const TArray<Json::TJsonReference<OpenAPI::V3::FParameterObject>>& InParameters, UWebAPISchema* OutSchema)
-			{
-				bool bAllConverted = true;
-				for (const Json::TJsonReference<OpenAPI::V3::FParameterObject>& Item : InParameters)
-				{
-					bAllConverted &= ConvertParameter(ResolveReference(Item)) != nullptr;
-				}
-
-				// Re-bind models to their TypeInfos
-				for (const TObjectPtr<UWebAPIModelBase>& Model : OutputSchema->Models)
-				{
-					Model->BindToTypeInfo();
-				}
+			return true;
+		}
 
-				return bAllConverted;
-			}
+		return false;
+	}
 
-			bool FWebAPIOpenAPISchemaConverter::ConvertParameters(const TMap<FString, TSharedPtr<OpenAPI::V3::FParameterObject>>& InParameters, UWebAPISchema* OutSchema)
-			{
-				bool bAllConverted = true;
-				for (const TTuple<FString, TSharedPtr<OpenAPI::V3::FParameterObject>>& Item : InParameters)
-				{
-					const FWebAPITypeNameVariant& Name = Item.Key;
-					const TSharedPtr<OpenAPI::V3::FParameterObject>& SrcParameter = Item.Value;
+	bool FWebAPIOpenAPISchemaConverter::ConvertTags(const TArray<TSharedPtr<OpenAPI::V3::FTagObject>>& InTags, UWebAPISchema* OutSchema) const
+	{
+		for (const TSharedPtr<OpenAPI::V3::FTagObject>& Tag : InTags)
+		{
+			ConvertService(Tag);
+		}
 
-					bAllConverted &= ConvertParameter(SrcParameter) != nullptr;
-				}
+		return OutSchema->Services.Num() > 0;
+	}
 
-				// Re-bind models to their TypeInfos
-				for (const TObjectPtr<UWebAPIModelBase>& Model : OutputSchema->Models)
-				{
-					Model->BindToTypeInfo();
-				}
+	bool FWebAPIOpenAPISchemaConverter::ConvertPaths(const TMap<FString, TSharedPtr<OpenAPI::V3::FPathItemObject>>& InPaths, UWebAPISchema* OutSchema)
+	{
+		for (const TTuple<FString, TSharedPtr<OpenAPI::V3::FPathItemObject>>& Item : InPaths)
+		{
+			const FString& Url = Item.Key;
+			const TSharedPtr<OpenAPI::V3::FPathItemObject>& Path = Item.Value;
 
-				return bAllConverted;
+			ConvertParameters(Path->Parameters, OutputSchema.Get());
+			
+			TMap<FString, TSharedPtr<OpenAPI::V3::FOperationObject>> SrcVerbs;
+			if (Path->Get.IsValid())
+			{
+				SrcVerbs.Add(TEXT("Get"), Path->Get);
+			}
+			if (Path->Put.IsValid())
+			{
+				SrcVerbs.Add(TEXT("Put"), Path->Put);
+			}
+			if (Path->Post.IsValid())
+			{
+				SrcVerbs.Add(TEXT("Post"), Path->Post);
+			}
+			if (Path->Delete.IsValid())
+			{
+				SrcVerbs.Add(TEXT("Delete"), Path->Delete);
+			}
+			if (Path->Options.IsValid())
+			{
+				SrcVerbs.Add(TEXT("Options"), Path->Options);
+			}
+			if (Path->Head.IsValid())
+			{
+				SrcVerbs.Add(TEXT("Head"), Path->Head);
+			}
+			if (Path->Patch.IsValid())
+			{
+				SrcVerbs.Add(TEXT("Patch"), Path->Patch);
 			}
 
-			bool FWebAPIOpenAPISchemaConverter::ConvertSecurity(const UE::WebAPI::OpenAPI::V3::FOpenAPIObject& InOpenAPI, UWebAPISchema* OutSchema)
+			// Each path can have multiple, ie. Get, Put, Delete
+			for (const TPair<FString, TSharedPtr<OpenAPI::V3::FOperationObject>>& VerbOperationPair : SrcVerbs)
 			{
-				if (!InOpenAPI.Components->SecuritySchemes.IsEmpty())
+				FString Verb = VerbOperationPair.Key;
+				if (VerbOperationPair.Value == nullptr || !VerbOperationPair.Value.IsValid())
 				{
-					for (const TPair<FString, Json::TJsonReference<OpenAPI::V3::FSecuritySchemeObject>>& Item : InOpenAPI.Components->SecuritySchemes)
-					{
-						const FString& Name = Item.Key;
-						const TSharedPtr<OpenAPI::V3::FSecuritySchemeObject>& SecurityScheme = ResolveReference(Item.Value);
-						if(SecurityScheme->Type == TEXT("http") && SecurityScheme->Scheme == TEXT("basic"))
-						{
-							// type = basic
-						}
-						else if(SecurityScheme->Type == TEXT("http") && SecurityScheme->Scheme == TEXT("bearer"))
-						{
-							// type = apiKey
-							// name = authorization
-							// storage = header
-						}
-						else if(SecurityScheme->Type == TEXT("oauth2"))
-						{
-							// @todo:
-							FString FlowName = SecurityScheme->Flows.IsValid() ? TEXT("") : TEXT("");
-							V3::FOAuthFlowsObject Flow = (SecurityScheme->Flows.Get())[0];
-
-							if(FlowName == TEXT("clientCredentials"))
-							{
-								// flow = application
-							}
-							else if(FlowName == TEXT("authorizationCode"))
-							{
-								// flow = accessCode
-							}
-							else
-							{
-								// flow = FlowName
-							}
+					continue;
+				}
 
+				TSharedPtr<OpenAPI::V3::FOperationObject> SrcOperation = VerbOperationPair.Value;
+				if(SrcOperation->RequestBody.IsSet())
+				{
+					FString ParameterDefinitionName;
+					// @todo:
+					/*
+					TSharedPtr<V3::FRequestBodyObject> Parameter = ResolveReference(SrcOperation->RequestBody.GetShared(), ParameterDefinitionName);
+					if(!Parameter->Content.IsEmpty())
+					{
+						auto Content = Parameter->Content;
 							
-						}
-
-						// @todo: WebAPI security primitive
-						// TObjectPtr<UWebAPIModel> Model = OutSchema->Client->Models.Add_GetRef(NewObject<UWebAPIModel>(OutSchema));
-						// ConvertModel(Name, Schema.Get(), Model);
 					}
-
-					return true;
+					*/
 				}
 
-				return false;
-			}
-
-			bool FWebAPIOpenAPISchemaConverter::ConvertTags(const TArray<TSharedPtr<OpenAPI::V3::FTagObject>>& InTags, UWebAPISchema* OutSchema) const
-			{
-				for (const TSharedPtr<OpenAPI::V3::FTagObject>& Tag : InTags)
+				TArray<FString> Tags;
+				if (SrcOperation.IsValid() && SrcOperation->Tags.IsSet())
 				{
-					ConvertService(Tag);
+					Tags = SrcOperation->Tags.GetValue();
 				}
-
-				return OutSchema->Services.Num() > 0;
-			}
-
-			bool FWebAPIOpenAPISchemaConverter::ConvertPaths(const TMap<FString, TSharedPtr<OpenAPI::V3::FPathItemObject>>& InPaths, UWebAPISchema* OutSchema)
-			{
-				for (const TTuple<FString, TSharedPtr<OpenAPI::V3::FPathItemObject>>& Item : InPaths)
+				else
 				{
-					const FString& Url = Item.Key;
-					const TSharedPtr<OpenAPI::V3::FPathItemObject>& Path = Item.Value;
-
-					ConvertParameters(Path->Parameters, OutputSchema.Get());
-					
-					TMap<FString, TSharedPtr<OpenAPI::V3::FOperationObject>> SrcVerbs;
-					if (Path->Get.IsValid())
-					{
-						SrcVerbs.Add(TEXT("Get"), Path->Get);
-					}
-					if (Path->Put.IsValid())
-					{
-						SrcVerbs.Add(TEXT("Put"), Path->Put);
-					}
-					if (Path->Post.IsValid())
-					{
-						SrcVerbs.Add(TEXT("Post"), Path->Post);
-					}
-					if (Path->Delete.IsValid())
-					{
-						SrcVerbs.Add(TEXT("Delete"), Path->Delete);
-					}
-					if (Path->Options.IsValid())
-					{
-						SrcVerbs.Add(TEXT("Options"), Path->Options);
-					}
-					if (Path->Head.IsValid())
-					{
-						SrcVerbs.Add(TEXT("Head"), Path->Head);
-					}
-					if (Path->Patch.IsValid())
-					{
-						SrcVerbs.Add(TEXT("Patch"), Path->Patch);
-					}
+					Tags = { TEXT("Default") };
+				}
 
-					// Each path can have multiple, ie. Get, Put, Delete
-					for (const TPair<FString, TSharedPtr<OpenAPI::V3::FOperationObject>>& VerbOperationPair : SrcVerbs)
-					{
-						FString Verb = VerbOperationPair.Key;
-						if (VerbOperationPair.Value == nullptr || !VerbOperationPair.Value.IsValid())
-						{
-							continue;
-						}
+				for (FString& Tag : Tags)
+				{
+					// Do first to retain original case
+					TObjectPtr<UWebAPIService> Service = ConvertService(Tag);
+					Tag = NameTransformer(Tag);
 
-						TSharedPtr<OpenAPI::V3::FOperationObject> SrcOperation = VerbOperationPair.Value;
-						if(SrcOperation->RequestBody.IsSet())
-						{
-							FString ParameterDefinitionName;
-							// @todo:
-							/*
-							TSharedPtr<V3::FRequestBodyObject> Parameter = ResolveReference(SrcOperation->RequestBody.GetShared(), ParameterDefinitionName);
-							if(!Parameter->Content.IsEmpty())
-							{
-								auto Content = Parameter->Content;
-									
-							}
-							*/
-						}
+					FString OperationName = SrcOperation->OperationId.Get(Verb + Tag);
+					OperationName = NameTransformer(OperationName);
+					OperationName = ProviderSettings.MakeValidMemberName(OperationName, Verb);
 
-						TArray<FString> Tags;
-						if (SrcOperation.IsValid() && SrcOperation->Tags.IsSet())
-						{
-							Tags = SrcOperation->Tags.GetValue();
-						}
-						else
-						{
-							Tags = { TEXT("Default") };
-						}
+					FString OperationNamePrefix = TEXT("");
+					int32 OperationNameSuffix = 0;
 
-						for (FString& Tag : Tags)
+					// Check if this already exists - each operation is unique so the name needs to be different - prepend verb
+					while(OutSchema->TypeRegistry->FindGeneratedType(EWebAPISchemaType::Operation,
+																	 OperationNamePrefix + OperationName + (OperationNameSuffix > 0 ? FString::FormatAsNumber(OperationNameSuffix) : TEXT(""))) != nullptr)
+					{
+						if(!OperationName.StartsWith(Verb))
 						{
-							// Do first to retain original case
-							TObjectPtr<UWebAPIService> Service = ConvertService(Tag);
-							Tag = NameTransformer(Tag);
-
-							FString OperationName = SrcOperation->OperationId.Get(Verb + Tag);
-							OperationName = NameTransformer(OperationName);
-							OperationName = ProviderSettings.MakeValidMemberName(OperationName, Verb);
-
-							FString OperationNamePrefix = TEXT("");
-							int32 OperationNameSuffix = 0;
-
-							// Check if this already exists - each operation is unique so the name needs to be different - prepend verb
-							while(OutSchema->TypeRegistry->FindGeneratedType(EWebAPISchemaType::Operation,
-																			 OperationNamePrefix + OperationName + (OperationNameSuffix > 0 ? FString::FormatAsNumber(OperationNameSuffix) : TEXT(""))) != nullptr)
-							{
-								if(!OperationName.StartsWith(Verb))
-								{
-									OperationNamePrefix = Verb;
-								}
-								OperationNameSuffix += 1; 
-							}
-
-							OperationName = OperationNamePrefix + OperationName + (OperationNameSuffix > 0 ? FString::FormatAsNumber(OperationNameSuffix) : TEXT(""));
-				
-							TObjectPtr<UWebAPITypeInfo> OperationTypeInfo = OutSchema->TypeRegistry->
-								GetOrMakeGeneratedType(
-													   EWebAPISchemaType::Operation,
-													   OperationName,
-													   OperationName,
-													   TEXT("U"));
-
-							const TObjectPtr<UWebAPIOperation> Operation = ConvertOperation(Url, Verb, SrcOperation, OperationTypeInfo);
-							Operation->Service = Service;
-							Operation->Verb = Verb;
-							Operation->Path = Url;
-
-							FStringFormatNamedArguments FormatArgs;
-							FormatArgs.Add(TEXT("ClassName"), UWebAPIOperation::StaticClass()->GetName());
-							FormatArgs.Add(TEXT("ServiceName"), Service->Name.ToString(true));
-							FormatArgs.Add(TEXT("OperationName"), OperationName);
-
-							const FName OperationObjectName = ProviderSettings.MakeOperationObjectName(Service, OperationName);
-							Operation->Rename(*OperationObjectName.ToString(), Operation->GetOuter());
-
-							Operation->BindToTypeInfo();
-
-							// Verb->Parameters
-							// Verb->Responses
-							// Verb->Schemes
-							// Verb->Consumes ie. "application/json"
-							// Verb->Produces ie. "application/json"
-							// Verb->Security							
+							OperationNamePrefix = Verb;
 						}
-						// @todo:
-						// ConvertParameters(Operation)
-
-						// @todo:
-						//ConvertResponses()
+						OperationNameSuffix += 1; 
 					}
+
+					OperationName = OperationNamePrefix + OperationName + (OperationNameSuffix > 0 ? FString::FormatAsNumber(OperationNameSuffix) : TEXT(""));
+		
+					TObjectPtr<UWebAPITypeInfo> OperationTypeInfo = OutSchema->TypeRegistry->
+						GetOrMakeGeneratedType(
+											   EWebAPISchemaType::Operation,
+											   OperationName,
+											   OperationName,
+											   TEXT("U"));
+
+					const TObjectPtr<UWebAPIOperation> Operation = ConvertOperation(Url, Verb, SrcOperation, OperationTypeInfo);
+					Operation->Service = Service;
+					Operation->Verb = Verb;
+					Operation->Path = Url;
+
+					FStringFormatNamedArguments FormatArgs;
+					FormatArgs.Add(TEXT("ClassName"), UWebAPIOperation::StaticClass()->GetName());
+					FormatArgs.Add(TEXT("ServiceName"), Service->Name.ToString(true));
+					FormatArgs.Add(TEXT("OperationName"), OperationName);
+
+					const FName OperationObjectName = ProviderSettings.MakeOperationObjectName(Service, OperationName);
+					Operation->Rename(*OperationObjectName.ToString(), Operation->GetOuter());
+
+					Operation->BindToTypeInfo();
+
+					// Verb->Parameters
+					// Verb->Responses
+					// Verb->Schemes
+					// Verb->Consumes ie. "application/json"
+					// Verb->Produces ie. "application/json"
+					// Verb->Security							
 				}
-				return true;
+				// @todo:
+				// ConvertParameters(Operation)
+
+				// @todo:
+				//ConvertResponses()
 			}
-		};
+		}
+		return true;
 	}
 };
 
diff --git a/Engine/Plugins/Web/WebAPI/Source/WebAPIOpenAPI/Private/V3/WebAPIOpenAPIConverter.h b/Engine/Plugins/Web/WebAPI/Source/WebAPIOpenAPI/Private/V3/WebAPIOpenAPIConverter.h
--- a/Engine/Plugins/Web/WebAPI/Source/WebAPIOpenAPI/Private/V3/WebAPIOpenAPIConverter.h
+++ b/Engine/Plugins/Web/WebAPI/Source/WebAPIOpenAPI/Private/V3/WebAPIOpenAPIConverter.h
@@ -14,134 +14,128 @@
 #include "UObject/StrongObjectPtr.h"
 #include "V3/WebAPIOpenAPISchema.h"
 
-namespace UE
+namespace UE::WebAPI::OpenAPI
 {
-	namespace WebAPI
+	class FWebAPIOpenAPISchemaConverter
 	{
-		namespace OpenAPI
-		{
-			class FWebAPIOpenAPISchemaConverter
-			{
-			public:
-				FWebAPIOpenAPISchemaConverter(
-					const TSharedPtr<const UE::WebAPI::OpenAPI::V3::FOpenAPIObject>& InOpenAPI,
-					UWebAPISchema* InWebAPISchema,
-					const TSharedRef<FWebAPIMessageLog>& InMessageLog,
-					const FWebAPIProviderSettings& InProviderSettings);
+	public:
+		FWebAPIOpenAPISchemaConverter(
+			const TSharedPtr<const UE::WebAPI::OpenAPI::V3::FOpenAPIObject>& InOpenAPI,
+			UWebAPISchema* InWebAPISchema,
+			const TSharedRef<FWebAPIMessageLog>& InMessageLog,
+			const FWebAPIProviderSettings& InProviderSettings);
 
-				bool Convert();
+		bool Convert();
 
-			public:
-				FString NameTransformer(const FWebAPINameVariant& InString) const;
+	public:
+		FString NameTransformer(const FWebAPINameVariant& InString) const;
 
-				TObjectPtr<UWebAPITypeInfo> ResolveMappedType(const FString& InType);
+		TObjectPtr<UWebAPITypeInfo> ResolveMappedType(const FString& InType);
 
-				TObjectPtr<UWebAPITypeInfo> ResolveType(FString InType, FString InFormat, FString InDefinitionName, const TSharedPtr<OpenAPI::V3::FSchemaObject>& InSchema = nullptr);
+		TObjectPtr<UWebAPITypeInfo> ResolveType(FString InType, FString InFormat, FString InDefinitionName, const TSharedPtr<OpenAPI::V3::FSchemaObject>& InSchema = nullptr);
 
-				template <typename SchemaType>
-				TObjectPtr<UWebAPITypeInfo> ResolveType(const TSharedPtr<SchemaType>& InSchema,	const FString& InDefinitionName = {}, FString InJsonType = {});
+		template <typename SchemaType>
+		TObjectPtr<UWebAPITypeInfo> ResolveType(const TSharedPtr<SchemaType>& InSchema,	const FString& InDefinitionName = {}, FString InJsonType = {});
 
-				TObjectPtr<UWebAPITypeInfo> GetTypeForContentType(const FString& InContentType);
+		TObjectPtr<UWebAPITypeInfo> GetTypeForContentType(const FString& InContentType);
 
-				template <typename StorageType>
-				FString GetDefaultJsonTypeForStorage(const StorageType& InStorage);
+		template <typename StorageType>
+		FString GetDefaultJsonTypeForStorage(const StorageType& InStorage);
 
-				template <typename ObjectType>
-				TSharedPtr<ObjectType> ResolveReference(const FString& InDefinitionName);
+		template <typename ObjectType>
+		TSharedPtr<ObjectType> ResolveReference(const FString& InDefinitionName);
 
-				template <>
-				TSharedPtr<OpenAPI::V3::FSchemaObject> ResolveReference(const FString& InDefinitionName);
+		/** Only resolves if necessary. */
+		template <typename ObjectType>
+		TSharedPtr<ObjectType> ResolveReference(const UE::Json::TJsonReference<ObjectType>& InJsonReference, FString& OutDefinitionName, bool bInCheck = true);
 
-				template <>
-				TSharedPtr<OpenAPI::V3::FParameterObject> ResolveReference(const FString& InDefinitionName);
+		/** Only resolves if necessary. */
+		template <typename ObjectType>
+		TSharedPtr<ObjectType> ResolveReference(const UE::Json::TJsonReference<ObjectType>& InJsonReference, bool bInCheck = true);
 
-				template <>
-				TSharedPtr<OpenAPI::V3::FResponseObject> ResolveReference(const FString& InDefinitionName);
+		//FWebAPINameVariant ResolveModelName(const TObjectPtr<UWebAPIProperty>& InProperty, const TOptional<bool>& bInIsArray, const FWebAPITypeNameVariant& InPossibleName = {});
 
-				template <>
-				TSharedPtr<OpenAPI::V3::FSecuritySchemeObject> ResolveReference(const FString& InDefinitionName);
+		FWebAPINameVariant ResolvePropertyName(const TObjectPtr<UWebAPIProperty>& InProperty, const FWebAPITypeNameVariant& InPotentialName, const TOptional<bool>& bInIsArray = {});
 
-				template <>
-				TSharedPtr<OpenAPI::V3::FRequestBodyObject> ResolveReference(const FString& InDefinitionName);
+		template <typename SchemaType>
+		bool IsArray(const TSharedPtr<SchemaType>& InSchema);
 
-				/** Only resolves if necessary. */
-				template <typename ObjectType>
-				TSharedPtr<ObjectType> ResolveReference(const UE::Json::TJsonReference<ObjectType>& InJsonReference, FString& OutDefinitionName, bool bInCheck = true);
+		template <typename SchemaType, typename ModelType>
+		static bool ConvertModelBase(const TSharedPtr<SchemaType>& InSchema, const TObjectPtr<ModelType>& OutModel);
 
-				/** Only resolves if necessary. */
-				template <typename ObjectType>
-				TSharedPtr<ObjectType> ResolveReference(const UE::Json::TJsonReference<ObjectType>& InJsonReference, bool bInCheck = true);
+		template <typename SchemaType>
+		TObjectPtr<UWebAPIEnum> ConvertEnum(const TSharedPtr<SchemaType>& InSrcSchema, const FWebAPITypeNameVariant& InEnumTypeName = {}) const;
 
-				//FWebAPINameVariant ResolveModelName(const TObjectPtr<UWebAPIProperty>& InProperty, const TOptional<bool>& bInIsArray, const FWebAPITypeNameVariant& InPossibleName = {});
+		/** Modifies an existing Property using the supplied source object. */
+		template <typename SchemaType>
+		bool ConvertProperty(const FWebAPITypeNameVariant& InModelName, const FWebAPINameVariant& InPropertyName, const TSharedPtr<SchemaType>& InSchema, const FString& InDefinitionName, const TObjectPtr<UWebAPIProperty>& OutProperty);
 
-				FWebAPINameVariant ResolvePropertyName(const TObjectPtr<UWebAPIProperty>& InProperty, const FWebAPITypeNameVariant& InPotentialName, const TOptional<bool>& bInIsArray = {});
-
-				template <typename SchemaType>
-				bool IsArray(const TSharedPtr<SchemaType>& InSchema);
+		/** Modifies an existing Property using the supplied TypeInfo. */
+		bool ConvertProperty(const FWebAPITypeNameVariant& InModelName, const FWebAPINameVariant& InPropertyName, const FWebAPITypeNameVariant& InPropertyTypeName, const TObjectPtr<UWebAPIProperty>& OutProperty);
+		
+		/** Creates a new Property from the supplied source object. */
+		template <typename SchemaType>
+		TObjectPtr<UWebAPIProperty> ConvertProperty(const TSharedPtr<SchemaType>& InSrcSchema, const TObjectPtr<UWebAPIModel>& InModel, const FWebAPINameVariant& InPropertyName = {}, const FString& InDefinitionName = {});
+/*
+		template <>
+		bool ConvertProperty(const FWebAPITypeNameVariant& InModelName,	const FWebAPINameVariant& InPropertyName, const TSharedPtr<OpenAPI::V2::FParameter>& InSchema, const FString& InDefinitionName, const TObjectPtr<UWebAPIProperty>& OutProperty);
+*/
+		/** Modifies an existing Model using the supplied source object. */
+		template <typename SchemaType>
+		bool ConvertModel(const TSharedPtr<SchemaType>& InSrcSchema, const FWebAPITypeNameVariant& InModelTypeName, const TObjectPtr<UWebAPIModel>& OutModel);
+		
+		/** Creates a new Model from the supplied source object. */
+		template <typename SchemaType>
+		TObjectPtr<UWebAPIModel> ConvertModel(const TSharedPtr<SchemaType>& InSrcSchema, const FWebAPITypeNameVariant& InModelTypeName = {});
 
-				template <typename SchemaType, typename ModelType>
-				static bool ConvertModelBase(const TSharedPtr<SchemaType>& InSchema, const TObjectPtr<ModelType>& OutModel);
+		bool ConvertOperationParameter(const FWebAPINameVariant& InParameterName, const TSharedPtr<OpenAPI::V3::FParameterObject>& InParameter, const FString& InDefinitionName, const TObjectPtr<UWebAPIOperationParameter>& OutParameter);
 
-				template <typename SchemaType>
-				TObjectPtr<UWebAPIEnum> ConvertEnum(const TSharedPtr<SchemaType>& InSrcSchema, const FWebAPITypeNameVariant& InEnumTypeName = {}) const;
+		TObjectPtr<UWebAPIParameter> ConvertParameter(const TSharedPtr<OpenAPI::V3::FParameterObject>& InSrcParameter);
+		
+		bool ConvertRequest(const FWebAPITypeNameVariant& InOperationName, const TSharedPtr<OpenAPI::V3::FOperationObject>& InOperation, const TObjectPtr<UWebAPIOperationRequest>& OutRequest);
 
-				/** Modifies an existing Property using the supplied source object. */
-				template <typename SchemaType>
-				bool ConvertProperty(const FWebAPITypeNameVariant& InModelName, const FWebAPINameVariant& InPropertyName, const TSharedPtr<SchemaType>& InSchema, const FString& InDefinitionName, const TObjectPtr<UWebAPIProperty>& OutProperty);
+		bool ConvertResponse(const FWebAPITypeNameVariant& InOperationName, uint32 InResponseCode, const TSharedPtr<OpenAPI::V3::FResponseObject>& InResponse, const TObjectPtr<UWebAPIOperationResponse>& OutResponse);
 
-				/** Modifies an existing Property using the supplied TypeInfo. */
-				bool ConvertProperty(const FWebAPITypeNameVariant& InModelName, const FWebAPINameVariant& InPropertyName, const FWebAPITypeNameVariant& InPropertyTypeName, const TObjectPtr<UWebAPIProperty>& OutProperty);
-				
-				/** Creates a new Property from the supplied source object. */
-				template <typename SchemaType>
-				TObjectPtr<UWebAPIProperty> ConvertProperty(const TSharedPtr<SchemaType>& InSrcSchema, const TObjectPtr<UWebAPIModel>& InModel, const FWebAPINameVariant& InPropertyName = {}, const FString& InDefinitionName = {});
-/*
-				template <>
-				bool ConvertProperty(const FWebAPITypeNameVariant& InModelName,	const FWebAPINameVariant& InPropertyName, const TSharedPtr<OpenAPI::V2::FParameter>& InSchema, const FString& InDefinitionName, const TObjectPtr<UWebAPIProperty>& OutProperty);
-*/
-				/** Modifies an existing Model using the supplied source object. */
-				template <typename SchemaType>
-				bool ConvertModel(const TSharedPtr<SchemaType>& InSrcSchema, const FWebAPITypeNameVariant& InModelTypeName, const TObjectPtr<UWebAPIModel>& OutModel);
+		TObjectPtr<UWebAPIOperation> ConvertOperation(const FString& InPath, const FString& InVerb, const TSharedPtr<OpenAPI::V3::FOperationObject>& InSrcOperation, const FWebAPITypeNameVariant& InOperationTypeName = {});
 
-				/** Modifies an existing Model using the supplied source object. */
-				template <>
-				bool ConvertModel(const TSharedPtr<OpenAPI::V3::FParameterObject>& InSrcSchema, const FWebAPITypeNameVariant& InModelTypeName, const TObjectPtr<UWebAPIModel>& OutModel);
+		TObjectPtr<UWebAPIService> ConvertService(const FWebAPINameVariant& InName) const;
+		TObjectPtr<UWebAPIService> ConvertService(const TSharedPtr<OpenAPI::V3::FTagObject>& InTag) const;
 
-				/** Creates a new Model from the supplied source object. */
-				template <typename SchemaType>
-				TObjectPtr<UWebAPIModel> ConvertModel(const TSharedPtr<SchemaType>& InSrcSchema, const FWebAPITypeNameVariant& InModelTypeName = {});
+		bool ConvertModels(const TMap<FString, Json::TJsonReference<OpenAPI::V3::FSchemaObject>>& InSchemas, UWebAPISchema* OutSchema);
 
-				bool ConvertOperationParameter(const FWebAPINameVariant& InParameterName, const TSharedPtr<OpenAPI::V3::FParameterObject>& InParameter, const FString& InDefinitionName, const TObjectPtr<UWebAPIOperationParameter>& OutParameter);
+		bool ConvertParameters(const TArray<Json::TJsonReference<OpenAPI::V3::FParameterObject>>& InParameters, UWebAPISchema* OutSchema);
 
-				TObjectPtr<UWebAPIParameter> ConvertParameter(const TSharedPtr<OpenAPI::V3::FParameterObject>& InSrcParameter);
-				
-				bool ConvertRequest(const FWebAPITypeNameVariant& InOperationName, const TSharedPtr<OpenAPI::V3::FOperationObject>& InOperation, const TObjectPtr<UWebAPIOperationRequest>& OutRequest);
+		bool ConvertParameters(const TMap<FString, TSharedPtr<OpenAPI::V3::FParameterObject>>& InParameters, UWebAPISchema* OutSchema);
 
-				bool ConvertResponse(const FWebAPITypeNameVariant& InOperationName, uint32 InResponseCode, const TSharedPtr<OpenAPI::V3::FResponseObject>& InResponse, const TObjectPtr<UWebAPIOperationResponse>& OutResponse);
+		bool ConvertSecurity(const UE::WebAPI::OpenAPI::V3::FOpenAPIObject& InOpenAPI, UWebAPISchema* OutSchema);
 
-				TObjectPtr<UWebAPIOperation> ConvertOperation(const FString& InPath, const FString& InVerb, const TSharedPtr<OpenAPI::V3::FOperationObject>& InSrcOperation, const FWebAPITypeNameVariant& InOperationTypeName = {});
+		bool ConvertTags(const TArray<TSharedPtr<OpenAPI::V3::FTagObject>>& InTags, UWebAPISchema* OutSchema) const;
 
-				TObjectPtr<UWebAPIService> ConvertService(const FWebAPINameVariant& InName) const;
-				TObjectPtr<UWebAPIService> ConvertService(const TSharedPtr<OpenAPI::V3::FTagObject>& InTag) const;
+		bool ConvertPaths(const TMap<FString, TSharedPtr<OpenAPI::V3::FPathItemObject>>& InPaths, UWebAPISchema* OutSchema);
 
-				bool ConvertModels(const TMap<FString, Json::TJsonReference<OpenAPI::V3::FSchemaObject>>& InSchemas, UWebAPISchema* OutSchema);
+	private:
+		TSharedPtr<const UE::WebAPI::OpenAPI::V3::FOpenAPIObject> InputSchema;
+		TStrongObjectPtr<UWebAPISchema> OutputSchema;
+		TSharedPtr<FWebAPIMessageLog> MessageLog;
+		FWebAPIProviderSettings ProviderSettings;
+	};
 
-				bool ConvertParameters(const TArray<Json::TJsonReference<OpenAPI::V3::FParameterObject>>& InParameters, UWebAPISchema* OutSchema);
+	template <>
+	TSharedPtr<OpenAPI::V3::FSchemaObject> FWebAPIOpenAPISchemaConverter::ResolveReference(const FString& InDefinitionName);
 
-				bool ConvertParameters(const TMap<FString, TSharedPtr<OpenAPI::V3::FParameterObject>>& InParameters, UWebAPISchema* OutSchema);
+	template <>
+	TSharedPtr<OpenAPI::V3::FParameterObject> FWebAPIOpenAPISchemaConverter::ResolveReference(const FString& InDefinitionName);
 
-				bool ConvertSecurity(const UE::WebAPI::OpenAPI::V3::FOpenAPIObject& InOpenAPI, UWebAPISchema* OutSchema);
+	template <>
+	TSharedPtr<OpenAPI::V3::FResponseObject> FWebAPIOpenAPISchemaConverter::ResolveReference(const FString& InDefinitionName);
 
-				bool ConvertTags(const TArray<TSharedPtr<OpenAPI::V3::FTagObject>>& InTags, UWebAPISchema* OutSchema) const;
+	template <>
+	TSharedPtr<OpenAPI::V3::FSecuritySchemeObject> FWebAPIOpenAPISchemaConverter::ResolveReference(const FString& InDefinitionName);
 
-				bool ConvertPaths(const TMap<FString, TSharedPtr<OpenAPI::V3::FPathItemObject>>& InPaths, UWebAPISchema* OutSchema);
+	template <>
+	TSharedPtr<OpenAPI::V3::FRequestBodyObject> FWebAPIOpenAPISchemaConverter::ResolveReference(const FString& InDefinitionName);
 
-			private:
-				TSharedPtr<const UE::WebAPI::OpenAPI::V3::FOpenAPIObject> InputSchema;
-				TStrongObjectPtr<UWebAPISchema> OutputSchema;
-				TSharedPtr<FWebAPIMessageLog> MessageLog;
-				FWebAPIProviderSettings ProviderSettings;
-			};
-		}
-	}
-}
+	/** Modifies an existing Model using the supplied source object. */
+	template <>
+	bool FWebAPIOpenAPISchemaConverter::ConvertModel(const TSharedPtr<OpenAPI::V3::FParameterObject>& InSrcSchema, const FWebAPITypeNameVariant& InModelTypeName, const TObjectPtr<UWebAPIModel>& OutModel);
+};
diff --git a/Engine/Plugins/Web/WebAPI/Source/WebAPIOpenAPI/Private/V3/WebAPIOpenAPIConverter.inl b/Engine/Plugins/Web/WebAPI/Source/WebAPIOpenAPI/Private/V3/WebAPIOpenAPIConverter.inl
new file mode 100644
--- /dev/null
+++ b/Engine/Plugins/Web/WebAPI/Source/WebAPIOpenAPI/Private/V3/WebAPIOpenAPIConverter.inl
@@ -0,0 +1,686 @@
+﻿// Copyright Epic Games, Inc. All Rights Reserved.
+
+#pragma once
+
+#include "WebAPIOpenAPIConverter.h"
+
+#include "IWebAPIEditorModule.h"
+#include "V3/WebAPIOpenAPIProvider.h"
+
+#define LOCTEXT_NAMESPACE "WebAPIOpenAPIConverter"
+
+namespace UE::WebAPI::OpenAPI
+{
+	template <>
+	FString FWebAPIOpenAPISchemaConverter::GetDefaultJsonTypeForStorage<EWebAPIParameterStorage>(const EWebAPIParameterStorage& InStorage)
+	{
+		static TMap<EWebAPIParameterStorage, FString> StorageToJsonType = {
+			{ EWebAPIParameterStorage::Body, TEXT("object") },
+			{ EWebAPIParameterStorage::Cookie, TEXT("string") },
+			{ EWebAPIParameterStorage::Header, TEXT("string") },
+			{ EWebAPIParameterStorage::Path, TEXT("string") },
+			{ EWebAPIParameterStorage::Query, TEXT("string") },
+		};
+
+		return StorageToJsonType[InStorage];
+	}
+
+	template <>
+	FString FWebAPIOpenAPISchemaConverter::GetDefaultJsonTypeForStorage<EWebAPIResponseStorage>(const EWebAPIResponseStorage& InStorage)
+	{
+		static TMap<EWebAPIResponseStorage, FString> StorageToJsonType = {
+			{ EWebAPIResponseStorage::Body, TEXT("object") },
+			{ EWebAPIResponseStorage::Header, TEXT("string") },
+		};
+
+		return StorageToJsonType[InStorage];
+	}
+
+	template <typename SchemaType>
+	TObjectPtr<UWebAPITypeInfo> FWebAPIOpenAPISchemaConverter::ResolveType(const TSharedPtr<SchemaType>& InSchema, const FString& InDefinitionName, FString InJsonType)
+	{
+		static_assert(TIsDerivedFrom<SchemaType, OpenAPI::V3::FSchemaObjectBase>::Value, "Type is not derived from OpenAPI::V3::FSchemaObjectBase.");
+
+		const TObjectPtr<UWebAPIStaticTypeRegistry> StaticTypeRegistry = IWebAPIEditorModuleInterface::Get().GetStaticTypeRegistry();
+
+		FString DefinitionName = InDefinitionName;
+		TObjectPtr<UWebAPITypeInfo> Result = nullptr;
+		
+		TSharedPtr<UE::WebAPI::OpenAPI::V3::FSchemaObject> ItemSchema = nullptr;
+		if(InSchema.IsValid() && InSchema->Items.IsSet())
+		{
+			ItemSchema = ResolveReference(InSchema->Items.GetValue(), DefinitionName);
+			if(!InSchema->Items->GetPath().IsEmpty())
+			{
+				DefinitionName = InSchema->Items->GetLastPathSegment();						
+			}
+			return ResolveType(ItemSchema, DefinitionName);
+		}
+
+		if(InSchema.IsValid())
+		{
+			Result = ResolveType(InSchema->Type.Get(InJsonType.IsEmpty() ? TEXT("object") : InJsonType), InSchema->Format.Get(TEXT("")), DefinitionName);
+		}
+		else
+		{
+			Result = ResolveType(InJsonType, TEXT(""), DefinitionName);
+			if(!Result)
+			{
+				Result = StaticTypeRegistry->Object;
+			}
+		}
+
+		if (InSchema.IsValid() && InSchema->Enum.IsSet() && !InSchema->Enum.GetValue().IsEmpty())
+		{
+			if(const TObjectPtr<UWebAPITypeInfo>* FoundGeneratedType = OutputSchema->TypeRegistry->FindGeneratedType(EWebAPISchemaType::Model, NameTransformer(DefinitionName)))
+			{
+				Result = *FoundGeneratedType;						
+			}
+			else
+			{
+				Result = OutputSchema->TypeRegistry->GetOrMakeGeneratedType(
+					EWebAPISchemaType::Model,
+					NameTransformer(DefinitionName),
+					DefinitionName,
+					StaticTypeRegistry->Enum);
+			}
+		}
+		else if (!Result->IsEnum() && (Result == StaticTypeRegistry->Object || Result->ToString(true).IsEmpty()))
+		{
+			if (!DefinitionName.IsEmpty())
+			{
+				if (const TObjectPtr<UWebAPITypeInfo>* FoundBuiltinType = StaticTypeRegistry->FindBuiltinType(DefinitionName))
+				{
+					Result = *FoundBuiltinType;
+				}
+				else if (const TObjectPtr<UWebAPITypeInfo>* FoundGeneratedModelType = OutputSchema->TypeRegistry->FindGeneratedType(EWebAPISchemaType::Model, DefinitionName))
+				{
+					Result = *FoundGeneratedModelType;
+				}
+				else if (const TObjectPtr<UWebAPITypeInfo>* FoundGeneratedParameterType = OutputSchema->TypeRegistry->FindGeneratedType(EWebAPISchemaType::Parameter, DefinitionName))
+				{
+					Result = *FoundGeneratedParameterType;
+				}
+				else
+				{
+					FFormatNamedArguments Args;
+					Args.Add(TEXT("DefinitionName"), FText::FromString(DefinitionName));
+					MessageLog->LogInfo(FText::Format(LOCTEXT("CannotResolveType", "ResolveType (object) failed to find a matching type for definition \"{DefinitionName}\", creating a new one."), Args), FWebAPIOpenAPIProvider::LogName);
+
+					// @bug: where is the model created for this?
+					// @todo: should this always be of type model, vs parameter?
+					Result = OutputSchema->TypeRegistry->GetOrMakeGeneratedType(
+						EWebAPISchemaType::Model,
+						NameTransformer(DefinitionName),
+						{},
+						TEXT("F"));
+					Result->JsonType = UWebAPIStaticTypeRegistry::ToFromJsonType;
+				}
+			}
+		}
+
+		check(!Result.IsNull());
+		
+		return Result;
+	}
+
+	template <>
+	TSharedPtr<OpenAPI::V3::FSchemaObject> FWebAPIOpenAPISchemaConverter::ResolveReference<OpenAPI::V3::FSchemaObject>(const FString& InDefinitionName)
+	{
+		if(!InputSchema->Components.IsValid())
+		{
+			return nullptr;
+		}
+		
+		if (const Json::TJsonReference<OpenAPI::V3::FSchemaObject>* FoundDefinition = InputSchema->Components->Schemas.Find(InDefinitionName))
+		{
+			if(!FoundDefinition->IsSet())
+			{
+				return ResolveReference<OpenAPI::V3::FSchemaObject>(FoundDefinition->GetLastPathSegment());
+			}
+			return FoundDefinition->GetShared();
+		}
+
+		return nullptr;
+	}
+
+	template <>
+	TSharedPtr<OpenAPI::V3::FParameterObject> FWebAPIOpenAPISchemaConverter::ResolveReference<OpenAPI::V3::FParameterObject>(const FString& InDefinitionName)
+	{
+		if(!InputSchema->Components.IsValid())
+		{
+			return nullptr;
+		}
+		
+		if (const Json::TJsonReference<OpenAPI::V3::FParameterObject>* FoundDefinition = InputSchema->Components->Parameters.Find(InDefinitionName))
+		{
+			if(!FoundDefinition->IsSet())
+			{
+				return ResolveReference<OpenAPI::V3::FParameterObject>(FoundDefinition->GetLastPathSegment());
+			}
+			return FoundDefinition->GetShared();
+		}
+
+		return nullptr;
+	}
+
+	template <>
+	TSharedPtr<OpenAPI::V3::FResponseObject> FWebAPIOpenAPISchemaConverter::ResolveReference<V3::FResponseObject>(const FString& InDefinitionName)
+	{
+		if(!InputSchema->Components.IsValid())
+		{
+			return nullptr;
+		}
+		
+		if (const Json::TJsonReference<OpenAPI::V3::FResponseObject>* FoundDefinition = InputSchema->Components->Responses.Find(InDefinitionName))
+		{
+			if(!FoundDefinition->IsSet())
+			{
+				return ResolveReference<OpenAPI::V3::FResponseObject>(FoundDefinition->GetLastPathSegment());
+			}
+			return FoundDefinition->GetShared();
+		}
+
+		return nullptr;
+	}
+
+	template <>
+	TSharedPtr<OpenAPI::V3::FSecuritySchemeObject> FWebAPIOpenAPISchemaConverter::ResolveReference<V3::FSecuritySchemeObject>(const FString& InDefinitionName)
+	{
+		if(!InputSchema->Components.IsValid())
+		{
+			return nullptr;
+		}
+		
+		if (const Json::TJsonReference<OpenAPI::V3::FSecuritySchemeObject>* FoundDefinition = InputSchema->Components->SecuritySchemes.Find(InDefinitionName))
+		{
+			if(!FoundDefinition->IsSet())
+			{
+				return ResolveReference<OpenAPI::V3::FSecuritySchemeObject>(FoundDefinition->GetLastPathSegment());
+			}
+			return FoundDefinition->GetShared();
+		}
+
+		return nullptr;
+	}
+
+	template <>
+	TSharedPtr<OpenAPI::V3::FRequestBodyObject> FWebAPIOpenAPISchemaConverter::ResolveReference<V3::FRequestBodyObject>(const FString& InDefinitionName)
+	{
+		if (const Json::TJsonReference<OpenAPI::V3::FRequestBodyObject>* FoundDefinition = InputSchema->Components->RequestBodies.Find(InDefinitionName))
+		{
+			if(!FoundDefinition->IsSet())
+			{
+				return ResolveReference<OpenAPI::V3::FRequestBodyObject>(FoundDefinition->GetLastPathSegment());
+			}
+			return FoundDefinition->GetShared();
+		}
+
+		return nullptr;
+	}
+
+		template <typename ObjectType>
+	TSharedPtr<ObjectType> FWebAPIOpenAPISchemaConverter::ResolveReference(const Json::TJsonReference<ObjectType>& InJsonReference, FString& OutDefinitionName, bool bInCheck)
+	{
+		if (InJsonReference.IsSet())
+		{
+			return InJsonReference.GetShared();
+		}
+
+		if(!InJsonReference.IsValid())
+		{
+			return nullptr;
+		}
+		
+		FString DefinitionName = InJsonReference.GetLastPathSegment();
+		if (TSharedPtr<ObjectType> FoundDefinition = ResolveReference<ObjectType>(DefinitionName))
+		{
+			OutDefinitionName = DefinitionName; // Only set if found 
+			return FoundDefinition;
+		}
+
+		if(bInCheck)
+		{
+			FFormatNamedArguments Args;
+			Args.Add(TEXT("ReferencePath"), FText::FromString(InJsonReference.GetPath()));
+			MessageLog->LogWarning(FText::Format(LOCTEXT("CannotResolveJsonReference", "Couldn't resolve JsonReference \"{ReferencePath}\""), Args),	FWebAPIOpenAPIProvider::LogName);
+		}
+
+		return nullptr;
+	}
+
+	template <typename ObjectType>
+	TSharedPtr<ObjectType> FWebAPIOpenAPISchemaConverter::ResolveReference(const UE::Json::TJsonReference<ObjectType>& InJsonReference, bool bInCheck)
+	{
+		FString DefinitionName;
+		return ResolveReference(InJsonReference, DefinitionName, bInCheck);
+	}
+
+	template <>
+	bool FWebAPIOpenAPISchemaConverter::IsArray(const TSharedPtr<UE::WebAPI::OpenAPI::V3::FSchemaObject>& InSchema)
+	{
+		return InSchema->Type.Get(TEXT(""))	== TEXT("array");
+	}
+
+	template <>
+	bool FWebAPIOpenAPISchemaConverter::IsArray(const TSharedPtr<UE::WebAPI::OpenAPI::V3::FParameterObject>& InSchema)
+	{
+		FString FoundDefinitionName;
+		return InSchema->Schema.IsSet()
+			? IsArray(ResolveReference(InSchema->Schema.GetValue(), FoundDefinitionName))
+			: false;
+	}
+
+	template <typename SchemaType, typename ModelType>
+	bool FWebAPIOpenAPISchemaConverter::ConvertModelBase(const TSharedPtr<SchemaType>& InSchema,
+														const TObjectPtr<ModelType>& OutModel)
+	{
+		static_assert(TIsDerivedFrom<SchemaType, OpenAPI::V3::FSchemaObjectBase>::Value, "Type is not derived from OpenAPI::V3::FSchemaObjectBase.");
+/*
+		OutModel->Description = InSchema->Description.Get(TEXT(""));
+
+		SET_OPTIONAL(InSchema->bRequired, OutModel->bIsRequired)
+		
+		SET_OPTIONAL_FLAGGED(InSchema->Minimum, OutModel->MinimumValue, OutModel->bUseMinimumValue);
+		SET_OPTIONAL_FLAGGED(InSchema->Maximum, OutModel->MaximumValue, OutModel->bUseMaximumValue);
+
+		if(!OutModel->bUseMinimumValue)
+		{
+			SET_OPTIONAL_FLAGGED(InSchema->MinLength, OutModel->MinimumValue, OutModel->bUseMinimumValue);
+		}
+
+		if(!OutModel->bUseMaximumValue)
+		{
+			SET_OPTIONAL_FLAGGED(InSchema->MaxLength, OutModel->MaximumValue, OutModel->bUseMaximumValue);					
+		}
+		
+		SET_OPTIONAL_FLAGGED(InSchema->Pattern, OutModel->Pattern, OutModel->bUsePattern);
+		*/
+		return true;
+	}
+
+	template <typename SchemaType>
+	TObjectPtr<UWebAPIEnum> FWebAPIOpenAPISchemaConverter::ConvertEnum(const TSharedPtr<SchemaType>& InSrcSchema, const FWebAPITypeNameVariant& InEnumTypeName) const
+	{
+		static_assert(TIsDerivedFrom<SchemaType, OpenAPI::V3::FSchemaObjectBase>::Value, "Type is not derived from OpenAPI::V3::FSchemaObjectBase.");
+
+		FWebAPITypeNameVariant EnumTypeName;
+		if(InEnumTypeName.IsValid())
+		{
+			EnumTypeName = InEnumTypeName;
+		}
+		else
+		{
+			const FString EnumName = InSrcSchema->Name;
+			check(!EnumName.IsEmpty());
+
+			const TObjectPtr<UWebAPIStaticTypeRegistry> StaticTypeRegistry = IWebAPIEditorModuleInterface::Get().GetStaticTypeRegistry();
+			EnumTypeName = OutputSchema->TypeRegistry->GetOrMakeGeneratedType(
+				EWebAPISchemaType::Model,
+				NameTransformer(EnumName),
+				EnumName,
+				StaticTypeRegistry->Enum);
+		}
+
+		const TObjectPtr<UWebAPIEnum>& DstEnum = OutputSchema->AddEnum(EnumTypeName.TypeInfo.Get());
+		DstEnum->Name = EnumTypeName;
+		DstEnum->Description = InSrcSchema->Description.Get(TEXT(""));
+
+		const TObjectPtr<UWebAPIModelBase> ModelBase = Cast<UWebAPIModelBase>(DstEnum);
+		if (!ConvertModelBase(InSrcSchema, ModelBase))
+		{
+			return nullptr;
+		}
+
+		for (const FString& SrcEnumValue : InSrcSchema->Enum.GetValue())
+		{
+			const TObjectPtr<UWebAPIEnumValue> DstEnumValue = DstEnum->AddValue();
+			DstEnumValue->Name.NameInfo.Name = NameTransformer(SrcEnumValue);
+			DstEnumValue->Name.NameInfo.JsonName = SrcEnumValue;
+		}
+
+		DstEnum->BindToTypeInfo();
+		
+		return DstEnum;
+	}
+
+	template <>
+	bool FWebAPIOpenAPISchemaConverter::ConvertProperty<OpenAPI::V3::FParameterObject>(
+		const FWebAPITypeNameVariant& InModelName,
+		const FWebAPINameVariant& InPropertyName,
+		const TSharedPtr<OpenAPI::V3::FParameterObject>& InParameter,
+		const FString& InDefinitionName,
+		const TObjectPtr<UWebAPIProperty>& OutProperty)
+	{
+		const FString DefinitionName = !InDefinitionName.IsEmpty()
+										? InDefinitionName
+										: ProviderSettings.MakeNestedPropertyTypeName(InModelName, InPropertyName);
+
+		const TObjectPtr<UWebAPIModelBase> ModelBase = OutProperty;
+		if (!ConvertModelBase(InParameter, ModelBase))
+		{
+			return false;
+		}
+
+		OutProperty->bIsArray = IsArray(InParameter);
+		
+		OutProperty->Name = FWebAPINameInfo(NameTransformer(InPropertyName.ToString()), InPropertyName.GetJsonName(), OutProperty->Type);				
+		OutProperty->Name = ResolvePropertyName(OutProperty, InModelName, {});
+		OutProperty->bIsRequired = InParameter->bRequired.Get(false);
+		OutProperty->BindToTypeInfo();
+
+		if(InParameter->Schema.IsSet())
+		{
+			const TSharedPtr<V3::FSchemaObject> Schema = ResolveReference(InParameter->Schema.GetValue());
+			OutProperty->Type = ResolveType(Schema);
+			
+			// Add enum as it's own model, and reference it as this properties type
+			if (InParameter->Schema.GetValue()->Enum.IsSet() && !InParameter->Schema.GetValue()->Enum->IsEmpty())
+			{
+				const FWebAPITypeNameVariant EnumTypeName = OutProperty->Type;
+				EnumTypeName.TypeInfo->SetNested(InModelName);
+				FString EnumName = OutProperty->Name.ToString(true);
+
+				// Only make nested name if the model and property name aren't the same, otherwise you get "NameName"!
+				if(EnumName != InModelName)
+				{
+					EnumName = ProviderSettings.MakeNestedPropertyTypeName(InModelName, OutProperty->Name.ToString(true));
+				}
+				EnumTypeName.TypeInfo->SetName(EnumName);
+
+				const TObjectPtr<UWebAPIEnum>& Enum = ConvertEnum(Schema, EnumTypeName);
+
+				const FText LogMessage = FText::FormatNamed(
+					LOCTEXT("AddedImplicitEnumForPropertyOfModel", "Implicit enum created for property \"{PropertyName}\" of model \"{ModelName}\"."),
+					TEXT("PropertyName"), FText::FromString(*InPropertyName.ToString(true)),
+					TEXT("ModelName"), FText::FromString(*InModelName.ToString(true)));
+				MessageLog->LogInfo(LogMessage, FWebAPIOpenAPIProvider::LogName);
+			
+				Enum->Name.TypeInfo->DebugString += LogMessage.ToString();
+				Enum->Name.TypeInfo->JsonName = InPropertyName.GetJsonName();
+
+				OutProperty->Type = Enum->Name;
+				OutProperty->Type.TypeInfo->Model = Enum;
+			}
+		}
+
+		return true;
+	}
+
+	template <typename SchemaType>
+	bool FWebAPIOpenAPISchemaConverter::ConvertProperty(const FWebAPITypeNameVariant& InModelName,
+														const FWebAPINameVariant& InPropertyName,
+														const TSharedPtr<SchemaType>& InSchema,
+														const FString& InDefinitionName,
+														const TObjectPtr<UWebAPIProperty>& OutProperty)
+	{
+		static_assert(TIsDerivedFrom<SchemaType, OpenAPI::V3::FSchemaObjectBase>::Value, "Type is not derived from OpenAPI::V3::FSchemaObjectBase.");
+
+		FString DefinitionName = !InDefinitionName.IsEmpty()
+								? InDefinitionName
+								: ProviderSettings.MakeNestedPropertyTypeName(InModelName, InPropertyName);
+
+		TSharedPtr<UE::WebAPI::OpenAPI::V3::FSchemaObject> ItemSchema = InSchema;
+		if(InSchema->Items.IsSet())
+		{
+			if(!InSchema->Items->GetPath().IsEmpty())
+			{
+				ItemSchema = ResolveReference(InSchema->Items.GetValue(), DefinitionName);
+				DefinitionName = InSchema->Items->GetLastPathSegment();						
+			}
+		}
+		else
+		{
+			ItemSchema = InSchema;
+		}
+
+		OutProperty->Type = ResolveType<SchemaType>(ItemSchema, DefinitionName);
+		check(OutProperty->Type.IsValid());
+
+		const TObjectPtr<UWebAPIModelBase> ModelBase = OutProperty;
+		if (!ConvertModelBase(ItemSchema, ModelBase))
+		{
+			return false;
+		}
+
+		OutProperty->bIsArray = IsArray(InSchema);
+		
+		OutProperty->Name = FWebAPINameInfo(NameTransformer(InPropertyName.ToString()), InPropertyName.GetJsonName(), OutProperty->Type);
+		OutProperty->Name = ResolvePropertyName(OutProperty, OutProperty->Type.ToString(true), {});
+		OutProperty->bIsRequired = InSchema->bRequired.Get(false);
+		OutProperty->BindToTypeInfo();
+
+		// Add enum as it's own model, and reference it as this properties type
+		if (InSchema->Enum.IsSet() && !InSchema->Enum->IsEmpty())
+		{
+			const FWebAPITypeNameVariant EnumTypeName = OutProperty->Type;
+			EnumTypeName.TypeInfo->SetName(ProviderSettings.MakeNestedPropertyTypeName(InModelName, InPropertyName));
+
+			const TObjectPtr<UWebAPIEnum>& Enum = ConvertEnum(InSchema, EnumTypeName);
+
+			const FText LogMessage = FText::FormatNamed(
+				LOCTEXT("AddedImplicitEnumForPropertyOfModel", "Implicit enum created for property \"{PropertyName}\" of model \"{ModelName}\"."),
+				TEXT("PropertyName"), FText::FromString(*InPropertyName.ToString(true)),
+				TEXT("ModelName"), FText::FromString(*InModelName.ToString(true)));
+			MessageLog->LogInfo(LogMessage, FWebAPIOpenAPIProvider::LogName);
+
+			Enum->Name.TypeInfo->DebugString += LogMessage.ToString();
+			
+			Enum->Name.TypeInfo->JsonName = InPropertyName.GetJsonName();
+			Enum->Name.TypeInfo->SetNested(InModelName);
+
+			OutProperty->Type = Enum->Name;
+			OutProperty->Type.TypeInfo->Model = Enum;
+		}
+		// Add struct as it's own model, and reference it as this properties type
+		else if (OutProperty->Type.ToString(true).IsEmpty())
+		{
+			check(OutProperty->Type.HasTypeInfo());
+			
+			const FWebAPITypeNameVariant ModelTypeName = OutProperty->Type;
+			ModelTypeName.TypeInfo->SetName(ProviderSettings.MakeNestedPropertyTypeName(InModelName, InPropertyName));
+			ModelTypeName.TypeInfo->Prefix = TEXT("F");
+			ModelTypeName.TypeInfo->SetNested(InModelName);
+
+			const TObjectPtr<UWebAPIModel>& Model = OutputSchema->AddModel<UWebAPIModel>(ModelTypeName.HasTypeInfo() ? ModelTypeName.TypeInfo.Get() : nullptr);
+			ConvertModel(InSchema, {}, Model);
+
+			const FText LogMessage = FText::FormatNamed(
+				LOCTEXT("AddedImplicitModelForPropertyOfModel", "Implicit model created for property \"{PropertyName}\" of model \"{ModelName}\"."),
+				TEXT("PropertyName"), FText::FromString(*InPropertyName.ToString(true)),
+				TEXT("ModelName"), FText::FromString(*InModelName.ToString(true)));
+			MessageLog->LogInfo(LogMessage, FWebAPIOpenAPIProvider::LogName);
+			
+			Model->Name.TypeInfo->DebugString += LogMessage.ToString();
+			
+			Model->Name.TypeInfo->JsonName = InPropertyName.GetJsonName();
+			Model->Name.TypeInfo->JsonType = UWebAPIStaticTypeRegistry::ToFromJsonType;
+
+			OutProperty->Type = Model->Name;
+			OutProperty->Type.TypeInfo->Model = Model;
+		}
+		
+		OutProperty->Name = ResolvePropertyName(OutProperty, InModelName);
+
+		return true;
+	}
+
+		template <typename SchemaType>
+	TObjectPtr<UWebAPIProperty> FWebAPIOpenAPISchemaConverter::ConvertProperty(const TSharedPtr<SchemaType>& InSrcSchema, const TObjectPtr<UWebAPIModel>& InModel, const FWebAPINameVariant& InPropertyName, const FString& InDefinitionName)
+	{
+		static_assert(TIsDerivedFrom<SchemaType, OpenAPI::V3::FSchemaObjectBase>::Value, "Type is not derived from OpenAPI::V3::FSchemaObjectBase.");
+
+		const FWebAPINameVariant SrcPropertyName = InPropertyName;
+		FString SrcPropertyDefinitionName = ProviderSettings.MakeNestedPropertyTypeName(*InModel->Name.ToString(true), SrcPropertyName);
+		const TSharedPtr<OpenAPI::V3::FSchemaObject> SrcPropertyValue = ResolveReference(InSrcSchema, SrcPropertyDefinitionName);
+
+		if (SrcPropertyValue)
+		{
+			const TObjectPtr<UWebAPIProperty>& DstProperty = InModel->Properties.Add_GetRef(NewObject<UWebAPIProperty>(InModel));
+			ConvertProperty(InModel->Name,
+				FWebAPINameInfo(NameTransformer(SrcPropertyName), SrcPropertyName.GetJsonName(), InModel->Name),
+				SrcPropertyValue,
+				SrcPropertyDefinitionName,
+				DstProperty);
+			
+			return DstProperty;
+		}
+
+		return nullptr;
+	}
+
+	template <typename SchemaType>
+	bool FWebAPIOpenAPISchemaConverter::ConvertModel(const TSharedPtr<SchemaType>& InSrcSchema, const FWebAPITypeNameVariant& InModelTypeName, const TObjectPtr<UWebAPIModel>& OutModel)
+	{
+		static_assert(TIsDerivedFrom<SchemaType, OpenAPI::V3::FSchemaObjectBase>::Value, "Type is not derived from OpenAPI::V3::FSchemaObjectBase.");
+
+		FWebAPITypeNameVariant ModelTypeName;
+		if(InModelTypeName.IsValid())
+		{
+			ModelTypeName = InModelTypeName;
+		}
+		else if(OutModel->Name.HasTypeInfo())
+		{
+			ModelTypeName = OutModel->Name;
+		}
+		else
+		{
+			const FString ModelName = InSrcSchema->Name;
+			check(!ModelName.IsEmpty());
+
+			ModelTypeName = OutputSchema->TypeRegistry->GetOrMakeGeneratedType(
+				EWebAPISchemaType::Model,
+				NameTransformer(*ModelName),
+				*ModelName,
+				TEXT("F"));
+		}
+
+		if (!InSrcSchema.IsValid())
+		{
+			FFormatNamedArguments Args;
+			Args.Add(TEXT("ModelName"), FText::FromString(*ModelTypeName.ToString(true)));
+			MessageLog->LogWarning(
+				FText::Format(LOCTEXT("SchemaInvalid",
+						"The schema for model \"{ModelName}\" was invalid/null."),
+					Args),
+				FWebAPIOpenAPIProvider::LogName);
+			return false;
+		}
+
+		const TObjectPtr<UWebAPIModel>& Model = OutModel ? OutModel : OutputSchema->AddModel<UWebAPIModel>(ModelTypeName.TypeInfo.Get());
+
+		const TObjectPtr<UWebAPIModelBase> ModelBase = Model;
+		if (!ConvertModelBase(InSrcSchema, ModelBase))
+		{
+			return false;
+		}
+
+		Model->Name = ModelTypeName;
+
+		if (InSrcSchema->Properties.IsSet() && !InSrcSchema->Properties->IsEmpty())
+		{
+			for (const TTuple<FString, Json::TJsonReference<OpenAPI::V3::FSchemaObject>>& NamePropertyPair : InSrcSchema->Properties.GetValue())
+			{
+				const FString& SrcPropertyName = NamePropertyPair.Key;
+				
+				FString SrcPropertyDefinitionName = ProviderSettings.MakeNestedPropertyTypeName(*ModelTypeName.ToString(true), SrcPropertyName);
+				TSharedPtr<OpenAPI::V3::FSchemaObject> SrcPropertySchema = ResolveReference(NamePropertyPair.Value, SrcPropertyDefinitionName);
+
+				TObjectPtr<UWebAPIProperty>& DstProperty = Model->Properties.Add_GetRef(NewObject<UWebAPIProperty>(Model));
+				ConvertProperty(Model->Name,
+								FWebAPINameInfo(NameTransformer(SrcPropertyName), SrcPropertyName, ModelTypeName),
+								SrcPropertySchema,
+								SrcPropertyDefinitionName,
+								DstProperty);
+			}
+		}
+
+		Model->BindToTypeInfo();
+
+		return true;
+	}
+
+	template <>
+	bool FWebAPIOpenAPISchemaConverter::ConvertModel<OpenAPI::V3::FParameterObject>(const TSharedPtr<OpenAPI::V3::FParameterObject>& InSrcSchema, const FWebAPITypeNameVariant& InModelTypeName, const TObjectPtr<UWebAPIModel>& OutModel)
+	{
+		FWebAPITypeNameVariant ModelTypeName;
+		if(InModelTypeName.IsValid() && InModelTypeName.HasTypeInfo())
+		{
+			ModelTypeName = InModelTypeName;
+		}
+		else if(OutModel->Name.HasTypeInfo())
+		{
+			ModelTypeName = OutModel->Name;
+		}
+		else
+		{
+			const FString ModelName = InSrcSchema->Name;
+			check(!ModelName.IsEmpty());
+
+			ModelTypeName = OutputSchema->TypeRegistry->GetOrMakeGeneratedType(
+				EWebAPISchemaType::Model,
+				NameTransformer(*ModelName),
+				*ModelName,
+				TEXT("F"));
+		}
+
+		if (!InSrcSchema.IsValid())
+		{
+			FFormatNamedArguments Args;
+			Args.Add(TEXT("ModelName"), FText::FromString(*ModelTypeName.ToString(true)));
+			MessageLog->LogWarning(
+				FText::Format(LOCTEXT("SchemaInvalid",
+						"The schema for model \"{ModelName}\" was invalid/null."),
+					Args),
+				FWebAPIOpenAPIProvider::LogName);
+			return false;
+		}
+
+		const TObjectPtr<UWebAPIModel>& Model = OutModel ? OutModel : OutputSchema->AddModel<UWebAPIModel>(ModelTypeName.TypeInfo.Get());
+
+		const TObjectPtr<UWebAPIModelBase> ModelBase = Model;
+		if (!ConvertModelBase(InSrcSchema, ModelBase))
+		{
+			return false;
+		}
+
+		Model->Name = ModelTypeName;
+
+		Model->BindToTypeInfo();
+
+		return true;
+	}
+
+	template <typename SchemaType>
+	TObjectPtr<UWebAPIModel> FWebAPIOpenAPISchemaConverter::ConvertModel(const TSharedPtr<SchemaType>& InSrcSchema, const FWebAPITypeNameVariant& InModelTypeName)
+	{
+		check(InSrcSchema.IsValid());
+
+		FWebAPITypeNameVariant ModelTypeName;
+		if(InModelTypeName.IsValid() && ModelTypeName.HasTypeInfo())
+		{
+			ModelTypeName = InModelTypeName;
+		}
+		else
+		{
+			const FString ModelName = InSrcSchema->Name;
+			check(!ModelName.IsEmpty());
+
+			ModelTypeName = OutputSchema->TypeRegistry->GetOrMakeGeneratedType(
+				EWebAPISchemaType::Model,
+				NameTransformer(*ModelName),
+				*ModelName,
+				TEXT("F"));
+		}
+		
+		const TObjectPtr<UWebAPIModel>& Model = OutputSchema->AddModel<UWebAPIModel>(ModelTypeName.TypeInfo.Get());
+		if(ConvertModel(InSrcSchema, ModelTypeName, Model))
+		{
+			return Model;
+		}
+
+		return nullptr;
+	}
+};
+
+#undef LOCTEXT_NAMESPACE
