WebAPI CIS fixes

diff --git a/Engine/Plugins/Web/WebAPI/Source/WebAPI/Public/WebAPIHttpRequest.h b/Engine/Plugins/Web/WebAPI/Source/WebAPI/Public/WebAPIHttpRequest.h
--- a/Engine/Plugins/Web/WebAPI/Source/WebAPI/Public/WebAPIHttpRequest.h
+++ b/Engine/Plugins/Web/WebAPI/Source/WebAPI/Public/WebAPIHttpRequest.h
@@ -116,13 +116,13 @@ TMap<FString, FString> TRequest<PayloadType>::GetAllHeaders() const
 	OutMap.Reserve(Headers.Num());
 	for (const FString& Header : Headers)
 	{
-		FString* const Key = nullptr;
-		FString* const Value = nullptr;
+		FString* Key = nullptr;
+		FString* Value = nullptr;
 		Header.Split(TEXT(":"), Key, Value);
 		
 		check(Key);
 		check(Value);
-		
+
 		OutMap.Emplace(MoveTemp(*Key), MoveTemp(*Value));
 	}
 
diff --git a/Engine/Plugins/Web/WebAPI/Source/WebAPI/Public/WebAPIMessageResponse.h b/Engine/Plugins/Web/WebAPI/Source/WebAPI/Public/WebAPIMessageResponse.h
--- a/Engine/Plugins/Web/WebAPI/Source/WebAPI/Public/WebAPIMessageResponse.h
+++ b/Engine/Plugins/Web/WebAPI/Source/WebAPI/Public/WebAPIMessageResponse.h
@@ -15,5 +15,10 @@ struct WEBAPI_API FWebAPIMessageResponse
 
 	virtual ~FWebAPIMessageResponse() = default;
 
-	virtual const FText& GetMessage() const PURE_VIRTUAL(FWebAPIMessageResponse::GetMessage, return FText::GetEmpty(););
+	virtual const FText& GetMessage() const;
 };
+
+inline const FText& FWebAPIMessageResponse::GetMessage() const
+{
+	return FText::GetEmpty();
+}
diff --git a/Engine/Plugins/Web/WebAPI/Source/WebAPIOpenAPI/Private/V3/WebAPIOpenAPIConverter.cpp b/Engine/Plugins/Web/WebAPI/Source/WebAPIOpenAPI/Private/V3/WebAPIOpenAPIConverter.cpp
--- a/Engine/Plugins/Web/WebAPI/Source/WebAPIOpenAPI/Private/V3/WebAPIOpenAPIConverter.cpp
+++ b/Engine/Plugins/Web/WebAPI/Source/WebAPIOpenAPI/Private/V3/WebAPIOpenAPIConverter.cpp
@@ -427,9 +427,16 @@ namespace UE::WebAPI::OpenAPI
 
 	TObjectPtr<UWebAPIParameter> FWebAPIOpenAPISchemaConverter::ConvertParameter(const TSharedPtr<OpenAPI::V3::FParameterObject>& InSrcParameter)
 	{
-		// const FString* ParameterName = InputSchema->Components.Parameters.FindKey(InSrcParameter);
-		// check(ParameterName);
 		const FString* ParameterName = nullptr;
+		for (const TPair<FString, Json::TJsonReference<V3::FParameterObject>>& ParameterPair : InputSchema->Components->Parameters)
+		{
+			if(ParameterPair.Value.IsSet() && ParameterPair.Value.GetShared() == InSrcParameter)
+			{
+				ParameterName = &ParameterPair.Key;
+				break;
+			}
+		}
+		check(ParameterName);
 
 		const FString ParameterJsonName = InSrcParameter->Name.IsEmpty() ? *ParameterName : InSrcParameter->Name;
 
