WebAPI cis fix

diff --git a/Engine/Plugins/Web/WebAPI/Source/WebAPI/Public/WebAPIHttpRequest.h b/Engine/Plugins/Web/WebAPI/Source/WebAPI/Public/WebAPIHttpRequest.h
--- a/Engine/Plugins/Web/WebAPI/Source/WebAPI/Public/WebAPIHttpRequest.h
+++ b/Engine/Plugins/Web/WebAPI/Source/WebAPI/Public/WebAPIHttpRequest.h
@@ -123,7 +123,10 @@ TMap<FString, FString> TRequest<PayloadType>::GetAllHeaders() const
 		check(Key);
 		check(Value);
 
-		OutMap.Emplace(MoveTemp(*Key), MoveTemp(*Value));
+		if(Key && Value)
+		{
+			OutMap.Emplace(MoveTemp(*Key), MoveTemp(*Value));
+		}
 	}
 
 	return OutMap;
diff --git a/Engine/Plugins/Web/WebAPI/Source/WebAPIOpenAPI/Private/V3/WebAPIOpenAPIConverter.cpp b/Engine/Plugins/Web/WebAPI/Source/WebAPIOpenAPI/Private/V3/WebAPIOpenAPIConverter.cpp
--- a/Engine/Plugins/Web/WebAPI/Source/WebAPIOpenAPI/Private/V3/WebAPIOpenAPIConverter.cpp
+++ b/Engine/Plugins/Web/WebAPI/Source/WebAPIOpenAPI/Private/V3/WebAPIOpenAPIConverter.cpp
@@ -436,7 +436,13 @@ namespace UE::WebAPI::OpenAPI
 				break;
 			}
 		}
+
 		check(ParameterName);
+		
+		if(!ParameterName)
+		{
+			return nullptr;
+		}
 
 		const FString ParameterJsonName = InSrcParameter->Name.IsEmpty() ? *ParameterName : InSrcParameter->Name;
 
