Fixed UWebAPIDefinition class to use AutoCollapseCategories instead of CollapseCategories.

diff --git a/Engine/Plugins/Web/WebAPI/Source/WebAPIEditor/Public/WebAPIDefinition.h b/Engine/Plugins/Web/WebAPI/Source/WebAPIEditor/Public/WebAPIDefinition.h
--- a/Engine/Plugins/Web/WebAPI/Source/WebAPIEditor/Public/WebAPIDefinition.h
+++ b/Engine/Plugins/Web/WebAPI/Source/WebAPIEditor/Public/WebAPIDefinition.h
@@ -1,4 +1,4 @@
-﻿// Copyright Epic Games, Inc. All Rights Reserved.
+// Copyright Epic Games, Inc. All Rights Reserved.
 
 #pragma once
 
@@ -15,7 +15,7 @@ class FWebAPIMessageLog;
 class UWebAPISchema;
 
 /** The asset containing various options and the schema itself. */
-UCLASS(CollapseCategories = "ImportSettings")
+UCLASS(AutoCollapseCategories = "ImportSettings")
 class WEBAPIEDITOR_API UWebAPIDefinition : public UObject
 {
 	GENERATED_BODY()
