WebAPI json utilities cleanup, custom node fixes

diff --git a/Engine/Plugins/Web/WebAPI/Source/WebAPI/Private/Security/WebAPIAuthentication.cpp b/Engine/Plugins/Web/WebAPI/Source/WebAPI/Private/Security/WebAPIAuthentication.cpp
--- a/Engine/Plugins/Web/WebAPI/Source/WebAPI/Private/Security/WebAPIAuthentication.cpp
+++ b/Engine/Plugins/Web/WebAPI/Source/WebAPI/Private/Security/WebAPIAuthentication.cpp
@@ -23,22 +23,20 @@ bool FWebAPIOAuthSchemeHandler::HandleHttpRequest(TSharedPtr<IHttpRequest> InReq
 {
 	check(InSettings);
 
-	// @todo: cache!
-	UWebAPIOAuthSettings* OAuthSettings = nullptr;
-	InSettings->AuthenticationSettings.FindItemByClass<UWebAPIOAuthSettings>(&OAuthSettings);	
-	check(OAuthSettings);
+	const UWebAPIOAuthSettings* AuthSettings = GetAuthSettings(InSettings);	
+	check(AuthSettings);
 
-	if(!ensureMsgf(OAuthSettings->IsValid(), TEXT("Authentication settings are missing one or more required properties.")))
+	if(!ensureMsgf(AuthSettings->IsValid(), TEXT("Authentication settings are missing one or more required properties.")))
 	{
 		return false;
 	}
 
-	if(OAuthSettings->AccessToken.IsEmpty())
+	if(AuthSettings->AccessToken.IsEmpty())
 	{
 		return false;
 	}
 
-	InRequest->SetHeader(TEXT("Authorization"), FString::Printf(TEXT("%s %s"), *OAuthSettings->TokenType, *OAuthSettings->AccessToken));
+	InRequest->SetHeader(TEXT("Authorization"), FString::Printf(TEXT("%s %s"), *AuthSettings->TokenType, *AuthSettings->AccessToken));
 
 	return true;
 }
@@ -47,27 +45,25 @@ bool FWebAPIOAuthSchemeHandler::HandleHttpResponse(EHttpResponseCodes::Type InRe
 {
 	check(InSettings);
 
-	// @todo: cache!
-	UWebAPIOAuthSettings* OAuthSettings = nullptr;
-	InSettings->AuthenticationSettings.FindItemByClass<UWebAPIOAuthSettings>(&OAuthSettings);	
-	check(OAuthSettings);
+	UWebAPIOAuthSettings* AuthSettings = GetAuthSettings(InSettings);	
+	check(AuthSettings);
 	
-	if(!ensureMsgf(OAuthSettings->IsValid(), TEXT("Authentication settings are missing one or more required properties.")))
+	if(!ensureMsgf(AuthSettings->IsValid(), TEXT("Authentication settings are missing one or more required properties.")))
 	{
 		return false;
 	}
 
-	GEngine->GetEngineSubsystem<UWebAPISubsystem>()->MakeHttpRequest(TEXT("POST"), [&, OAuthSettings](const TSharedPtr<IHttpRequest>& InRequest)
+	GEngine->GetEngineSubsystem<UWebAPISubsystem>()->MakeHttpRequest(TEXT("POST"), [&, AuthSettings](const TSharedPtr<IHttpRequest>& InRequest)
 	{
 		//Microsoft login rest URL
 		FStringFormatNamedArguments UrlArgs;
-		UrlArgs.Add(TEXT("TenantId"), OAuthSettings->TenantId);
-		const FString Url = FString::Format(*OAuthSettings->AuthenticationServer, UrlArgs);
+		UrlArgs.Add(TEXT("TenantId"), AuthSettings->TenantId);
+		const FString Url = FString::Format(*AuthSettings->AuthenticationServer, UrlArgs);
 
 		// @todo: move!
 		FString StringPayload = "grant_type=client_credentials";
-		StringPayload.Append("&client_id=" + OAuthSettings->ClientId);
-		StringPayload.Append("&client_secret=" + OAuthSettings->ClientSecret);
+		StringPayload.Append("&client_id=" + AuthSettings->ClientId);
+		StringPayload.Append("&client_secret=" + AuthSettings->ClientSecret);
 		StringPayload.Append("&resource=https://digitaltwins.azure.net");
 
 		InRequest->SetURL(Url);
@@ -93,7 +89,7 @@ bool FWebAPIOAuthSchemeHandler::HandleHttpResponse(EHttpResponseCodes::Type InRe
 			InRequest->SetHeader(Header.Key, Header.Value);
 		}
 	})
-	.Next([this, OAuthSettings](const TTuple<FHttpResponsePtr, bool>& InResponse)
+	.Next([this, AuthSettings](const TTuple<FHttpResponsePtr, bool>& InResponse)
 	{
 		// Request failed
 		if(!InResponse.Get<bool>())
@@ -121,10 +117,10 @@ bool FWebAPIOAuthSchemeHandler::HandleHttpResponse(EHttpResponseCodes::Type InRe
 					const int32 UnixTimeExpire = JsonObject->GetNumberField("expires_on");
 
 					UE_LOG(LogWebAPI, Display, TEXT("Generate token Response success"));
-					OAuthSettings->TokenType = JsonObject->GetStringField("token_type");
-					OAuthSettings->AccessToken = JsonObject->GetStringField("access_token");
-					OAuthSettings->ExpiresOn =  FDateTime::FromUnixTimestamp(UnixTimeExpire);
-					OAuthSettings->SaveConfig();
+					AuthSettings->TokenType = JsonObject->GetStringField("token_type");
+					AuthSettings->AccessToken = JsonObject->GetStringField("access_token");
+					AuthSettings->ExpiresOn =  FDateTime::FromUnixTimestamp(UnixTimeExpire);
+					AuthSettings->SaveConfig();
 					bSuccess = true;
 				}
 				else
@@ -146,7 +142,7 @@ bool FWebAPIOAuthSchemeHandler::HandleHttpResponse(EHttpResponseCodes::Type InRe
 			Message = "Response is null";
 		}
 
-		// @todo:
+		// @todo: request retry
 		//Redo all the requests that arrived while fetching new token
 		/*
 		if (bUpdatingToken && bSuccess)
@@ -161,3 +157,17 @@ bool FWebAPIOAuthSchemeHandler::HandleHttpResponse(EHttpResponseCodes::Type InRe
 
 	return true;
 }
+
+UWebAPIOAuthSettings* FWebAPIOAuthSchemeHandler::GetAuthSettings(const UWebAPIDeveloperSettings* InSettings)
+{
+	if(!CachedAuthSettings.IsValid() && InSettings)
+	{
+		UWebAPIOAuthSettings* OAuthSettings = nullptr;
+		InSettings->AuthenticationSettings.FindItemByClass<UWebAPIOAuthSettings>(&OAuthSettings);
+		check(OAuthSettings);
+
+		CachedAuthSettings = OAuthSettings;	
+	}
+	
+	return CachedAuthSettings.IsValid() ? CachedAuthSettings.Get() : nullptr;
+}
diff --git a/Engine/Plugins/Web/WebAPI/Source/WebAPI/Public/Security/WebAPIAuthentication.h b/Engine/Plugins/Web/WebAPI/Source/WebAPI/Public/Security/WebAPIAuthentication.h
--- a/Engine/Plugins/Web/WebAPI/Source/WebAPI/Public/Security/WebAPIAuthentication.h
+++ b/Engine/Plugins/Web/WebAPI/Source/WebAPI/Public/Security/WebAPIAuthentication.h
@@ -9,6 +9,8 @@
 
 #include "WebAPIAuthentication.generated.h"
 
+class UWebAPIOAuthSettings;
+
 // @note: Each will have user supplied settings, and custom handling of auth requests
 
 /**
@@ -38,7 +40,7 @@ public:
 	
 	/**  */
 	UPROPERTY(Config, EditAnywhere, BlueprintReadWrite, Category = "Security")
-	FString TenantId; // API key? @note: this is specific to Azure services! move!
+	FString TenantId; // API key? @todo: this is specific to Azure services! move!
 
 	/** Public client identifier. */
 	UPROPERTY(Config, EditAnywhere, BlueprintReadWrite, Category = "Security")
@@ -85,4 +87,9 @@ class WEBAPI_API FWebAPIOAuthSchemeHandler
 public:
 	virtual bool HandleHttpRequest(TSharedPtr<IHttpRequest> InRequest, UWebAPIDeveloperSettings* InSettings) override;
 	virtual bool HandleHttpResponse(EHttpResponseCodes::Type InResponseCode, TSharedPtr<IHttpResponse> InResponse, bool bInWasSuccessful, UWebAPIDeveloperSettings* InSettings) override;
+
+private:
+	UWebAPIOAuthSettings* GetAuthSettings(const UWebAPIDeveloperSettings* InSettings);
+	
+	TWeakObjectPtr<UWebAPIOAuthSettings> CachedAuthSettings;
 };
diff --git a/Engine/Plugins/Web/WebAPI/Source/WebAPIBlueprintGraph/Private/K2Node_WebAPIOperation.cpp b/Engine/Plugins/Web/WebAPI/Source/WebAPIBlueprintGraph/Private/K2Node_WebAPIOperation.cpp
--- a/Engine/Plugins/Web/WebAPI/Source/WebAPIBlueprintGraph/Private/K2Node_WebAPIOperation.cpp
+++ b/Engine/Plugins/Web/WebAPI/Source/WebAPIBlueprintGraph/Private/K2Node_WebAPIOperation.cpp
@@ -3,45 +3,289 @@
 #include "K2Node_WebAPIOperation.h"
 
 #include "BlueprintActionDatabaseRegistrar.h"
+#include "BlueprintEditor.h"
 #include "BlueprintTypePromotion.h"
 #include "EdGraphSchema_K2.h"
-#include "K2Node_StructOperation.h"
+#include "K2Node_AddDelegate.h"
+#include "K2Node_AssignmentStatement.h"
+#include "K2Node_CallFunction.h"
+#include "K2Node_CreateDelegate.h"
+#include "K2Node_CustomEvent.h"
+#include "K2Node_IfThenElse.h"
+#include "K2Node_Self.h"
+#include "K2Node_TemporaryVariable.h"
 #include "KismetCompiler.h"
+#include "PropertyPath.h"
 #include "ScopedTransaction.h"
 #include "ToolMenu.h"
+#include "WebAPIBlueprintGraphLog.h"
+#include "WebAPIBlueprintGraphUtilities.h"
 #include "WebAPIOperationObject.h"
 #include "Algo/AllOf.h"
 #include "EdGraph/EdGraphNode.h"
 #include "EdGraph/EdGraphPin.h"
+#include "Kismet/KismetSystemLibrary.h"
 #include "Kismet2/BlueprintEditorUtils.h"
+#include "Kismet2/KismetEditorUtilities.h"
 
 #define LOCTEXT_NAMESPACE "K2Node_WebAPIAsyncOperation"
 
-struct FGetPinName
+namespace UE::WebAPI::Private
 {
-	static const FName& GetMessagePinName() { static const FName MessagePinName(TEXT("Message")); return MessagePinName; }
-	static const FName& GetSuccessPinName() { static const FName SuccessPinName(TEXT("OnSuccess")); return SuccessPinName; }
-	static const FName& GetErrorPinName() { static const FName ErrorPinName(TEXT("OnError")); return ErrorPinName; }
-};
+	// @note: the next calls are from BaseAsyncTask
+	
+	static bool ValidDataPin(const UEdGraphPin* InPin, EEdGraphPinDirection InDirection)
+	{
+		const bool bValidDataPin = InPin
+			&& !InPin->bOrphanedPin
+			&& (InPin->PinType.PinCategory != UEdGraphSchema_K2::PC_Exec);
 
-UK2Node_WebAPIOperation::UK2Node_WebAPIOperation()
-{
-	ProxyActivateFunctionName = NAME_None;
+		const bool bProperDirection = InPin && (InPin->Direction == InDirection);
+
+		return bValidDataPin && bProperDirection;
+	}
+
+	bool CreateDelegateForNewFunction(UEdGraphPin* DelegateInputPin, FName FunctionName, UK2Node* CurrentNode, UEdGraph* SourceGraph, FKismetCompilerContext& CompilerContext)
+	{
+		const UEdGraphSchema_K2* Schema = CompilerContext.GetSchema();
+		check(DelegateInputPin && Schema && CurrentNode && SourceGraph && (FunctionName != NAME_None));
+		bool bResult = true;
+
+		// WORKAROUND, so we can create delegate from nonexistent function by avoiding check at expanding step
+		// instead simply: Schema->TryCreateConnection(AddDelegateNode->GetDelegatePin(), CurrentCENode->FindPinChecked(UK2Node_CustomEvent::DelegateOutputName));
+		UK2Node_Self* SelfNode = CompilerContext.SpawnIntermediateNode<UK2Node_Self>(CurrentNode, SourceGraph);
+		SelfNode->AllocateDefaultPins();
+
+		UK2Node_CreateDelegate* CreateDelegateNode = CompilerContext.SpawnIntermediateNode<UK2Node_CreateDelegate>(CurrentNode, SourceGraph);
+		CreateDelegateNode->AllocateDefaultPins();
+		bResult &= Schema->TryCreateConnection(DelegateInputPin, CreateDelegateNode->GetDelegateOutPin());
+		bResult &= Schema->TryCreateConnection(SelfNode->FindPinChecked(UEdGraphSchema_K2::PN_Self), CreateDelegateNode->GetObjectInPin());
+		CreateDelegateNode->SetFunction(FunctionName);
+
+		return bResult;
+	}
+
+	bool CopyEventSignature(UK2Node_CustomEvent* CENode, UFunction* Function, const UEdGraphSchema_K2* Schema)
+	{
+		check(CENode && Function && Schema);
 
-	// @todo:
-	// 1. on new node spawned, break Request (input) pin by default: FBlueprintEditor::OnSplitStructPin()
-	// 2. option to convert to/from callback styles: UK2Node_DynamicCast::SetPurity(bool bNewPurity)
-	// 3. collapse output responses that inherit MessageResponse to single message output pin
+		bool bResult = true;
+		for (TFieldIterator<FProperty> PropIt(Function); PropIt && (PropIt->PropertyFlags & CPF_Parm); ++PropIt)
+		{
+			const FProperty* Param = *PropIt;
+			if (!Param->HasAnyPropertyFlags(CPF_OutParm) || Param->HasAnyPropertyFlags(CPF_ReferenceParm))
+			{
+				FEdGraphPinType PinType;
+				bResult &= Schema->ConvertPropertyToPinType(Param, /*out*/ PinType);
+				bResult &= (nullptr != CENode->CreateUserDefinedPin(Param->GetFName(), PinType, EGPD_Output));
+			}
+		}
+		return bResult;
+	}
+	
+	bool HandleDelegateImplementation(
+		FMulticastDelegateProperty* CurrentProperty, const TArray<FOutputPinAndLocalVariable>& VariableOutputs,
+		UEdGraphPin* ProxyObjectPin, UEdGraphPin*& InOutLastThenPin, UEdGraphPin*& OutLastActivatedThenPin,
+		UK2Node* CurrentNode, UEdGraph* SourceGraph, FKismetCompilerContext& CompilerContext)
+	{
+		bool bIsErrorFree = true;
+		const UEdGraphSchema_K2* Schema = CompilerContext.GetSchema();
+		check(CurrentProperty && ProxyObjectPin && InOutLastThenPin && CurrentNode && SourceGraph && Schema);
+
+		UEdGraphPin* PinForCurrentDelegateProperty = CurrentNode->FindPin(CurrentProperty->GetFName());
+		if (!PinForCurrentDelegateProperty || (UEdGraphSchema_K2::PC_Exec != PinForCurrentDelegateProperty->PinType.PinCategory))
+		{
+			FText ErrorMessage = FText::Format(LOCTEXT("WrongDelegateProperty", "WebAPIOperation: Cannot find execution pin for delegate "), FText::FromString(CurrentProperty->GetName()));
+			CompilerContext.MessageLog.Error(*ErrorMessage.ToString(), CurrentNode);
+			return false;
+		}
+
+		UK2Node_CustomEvent* CurrentCENode = CompilerContext.SpawnIntermediateEventNode<UK2Node_CustomEvent>(CurrentNode, PinForCurrentDelegateProperty, SourceGraph);
+		{
+		UK2Node_AddDelegate* AddDelegateNode = CompilerContext.SpawnIntermediateNode<UK2Node_AddDelegate>(CurrentNode, SourceGraph);
+		AddDelegateNode->SetFromProperty(CurrentProperty, false, CurrentProperty->GetOwnerClass());
+		AddDelegateNode->AllocateDefaultPins();
+			bIsErrorFree &= Schema->TryCreateConnection(AddDelegateNode->FindPinChecked(UEdGraphSchema_K2::PN_Self), ProxyObjectPin);
+			bIsErrorFree &= Schema->TryCreateConnection(InOutLastThenPin, AddDelegateNode->FindPinChecked(UEdGraphSchema_K2::PN_Execute));
+			InOutLastThenPin = AddDelegateNode->FindPinChecked(UEdGraphSchema_K2::PN_Then);
+			CurrentCENode->CustomFunctionName = *FString::Printf(TEXT("%s_%s"), *CurrentProperty->GetName(), *CompilerContext.GetGuid(CurrentNode));
+			CurrentCENode->AllocateDefaultPins();
+
+			bIsErrorFree &= CreateDelegateForNewFunction(AddDelegateNode->GetDelegatePin(), CurrentCENode->GetFunctionName(), CurrentNode, SourceGraph, CompilerContext);
+			bIsErrorFree &= CopyEventSignature(CurrentCENode, AddDelegateNode->GetDelegateSignature(), Schema);
+		}
+
+		OutLastActivatedThenPin = CurrentCENode->FindPinChecked(UEdGraphSchema_K2::PN_Then);
+		for (const FOutputPinAndLocalVariable& OutputPair : VariableOutputs) // CREATE CHAIN OF ASSIGMENTS
+		{
+			UEdGraphPin* PinWithData = CurrentCENode->FindPin(OutputPair.OutputPin->PinName);
+			if (PinWithData == nullptr)
+			{
+				continue;
+			}
+
+			UK2Node_AssignmentStatement* AssignNode = CompilerContext.SpawnIntermediateNode<UK2Node_AssignmentStatement>(CurrentNode, SourceGraph);
+			AssignNode->AllocateDefaultPins();
+			bIsErrorFree &= Schema->TryCreateConnection(OutLastActivatedThenPin, AssignNode->GetExecPin());
+			bIsErrorFree &= Schema->TryCreateConnection(OutputPair.TempVar->GetVariablePin(), AssignNode->GetVariablePin());
+			AssignNode->NotifyPinConnectionListChanged(AssignNode->GetVariablePin());
+			bIsErrorFree &= Schema->TryCreateConnection(AssignNode->GetValuePin(), PinWithData);
+			AssignNode->NotifyPinConnectionListChanged(AssignNode->GetValuePin());
+
+			OutLastActivatedThenPin = AssignNode->GetThenPin();
+		}
+
+		bIsErrorFree &= CompilerContext.MovePinLinksToIntermediate(*PinForCurrentDelegateProperty, *OutLastActivatedThenPin).CanSafeConnect();
+		return bIsErrorFree;
+	}
 }
 
 void UK2Node_WebAPIOperation::AllocateDefaultPins()
 {
+	InvalidatePinTooltips();
+
+	const UEdGraphSchema_K2* K2Schema = GetDefault<UEdGraphSchema_K2>();
+
+	// Exec input/output
+	if(!GetExecPin())
+	{
+		CreatePin(EGPD_Input, UEdGraphSchema_K2::PC_Exec, UEdGraphSchema_K2::PN_Execute);
+	}
+
+	if(!GetThenPin())
+	{
+		CreatePin(EGPD_Output, UEdGraphSchema_K2::PC_Exec, UEdGraphSchema_K2::PN_Then);
+	}
+
+	const UFunction* Function = GetFactoryFunction();
+
+	TArray<FMulticastDelegateProperty*, TFixedAllocator<2>> DelegateProperties = {
+		UE::WebAPI::Operation::GetPositiveOutcomeDelegate(OperationClass),
+		UE::WebAPI::Operation::GetNegativeOutcomeDelegate(OperationClass)
+	};
+
+	if(OperationAsyncType == EWebAPIOperationAsyncType::LatentAction)
+	{
+		// Exec outputs for both outcomes
+		for(const FMulticastDelegateProperty* DelegateProperty : DelegateProperties)
+		{
+			UEdGraphPin* ExecPin = CreatePin(EGPD_Output, UEdGraphSchema_K2::PC_Exec, DelegateProperty->GetFName());
+			ExecPin->PinToolTip = DelegateProperty->GetToolTipText().ToString();
+			ExecPin->PinFriendlyName = DelegateProperty->GetDisplayNameText();
+		}
+
+		// Response pins
+		const UFunction* DelegateSignatureFunction = UE::WebAPI::Operation::GetOutcomeDelegateSignatureFunction(OperationClass);
+		if (DelegateSignatureFunction)
+		{
+			for (TFieldIterator<FProperty> ParameterIterator(DelegateSignatureFunction); ParameterIterator && (ParameterIterator->PropertyFlags & CPF_Parm); ++ParameterIterator)
+			{
+				const FProperty* Parameter = *ParameterIterator;
+				const bool bIsFunctionInput = !Parameter->HasAnyPropertyFlags(CPF_OutParm) || Parameter->HasAnyPropertyFlags(CPF_ReferenceParm);
+				if (bIsFunctionInput)
+				{
+					UEdGraphPin* Pin = CreatePin(EGPD_Output, NAME_None, Parameter->GetFName());
+					K2Schema->ConvertPropertyToPinType(Parameter, /*out*/ Pin->PinType);
+
+					// Check for a display name override
+					const FString& PinDisplayName = Parameter->GetMetaData(FBlueprintMetadata::MD_DisplayName);
+					if (!PinDisplayName.IsEmpty())
+					{
+						Pin->PinFriendlyName = FText::FromString(PinDisplayName);
+					}
+					// Else cleanup name
+					else
+					{
+						FString PinNameStr = Pin->PinName.ToString();
+						UE::WebAPI::Graph::CleanupPinNameInline(PinNameStr);
+						Pin->PinFriendlyName = FText::FromString(PinNameStr);
+					}
+
+					UK2Node_CallFunction::GeneratePinTooltipFromFunction(*Pin, DelegateSignatureFunction);
+
+					UE::WebAPI::Graph::SplitPin(Pin);
+				}
+			}
+		}
+	}
+
+	bool bAllPinsGood = true;
+	if (Function)
+	{
+		TSet<FName> PinsToHide;
+		FBlueprintEditorUtils::GetHiddenPinsForFunction(GetGraph(), Function, PinsToHide);
+
+		// Input pins
+		for (TFieldIterator<FProperty> ParameterIterator(Function); ParameterIterator && (ParameterIterator->PropertyFlags & CPF_Parm); ++ParameterIterator)
+		{
+			const FProperty* Parameter = *ParameterIterator;
+			const bool bIsFunctionInput = !Parameter->HasAnyPropertyFlags(CPF_OutParm) || Parameter->HasAnyPropertyFlags(CPF_ReferenceParm);
+			if (!bIsFunctionInput)
+			{
+				// skip function output, it's internal node data 
+				continue;
+			}
+
+			UEdGraphNode::FCreatePinParams PinParams;
+			PinParams.bIsReference = Parameter->HasAnyPropertyFlags(CPF_ReferenceParm) && bIsFunctionInput;
+			
+			UEdGraphPin* Pin = CreatePin(EGPD_Input, NAME_None, Parameter->GetFName(), PinParams);
+			const bool bPinGood = (Pin && K2Schema->ConvertPropertyToPinType(Parameter, /*out*/ Pin->PinType));
+
+			if (bPinGood)
+			{
+				// Check for a display name override
+				const FString& PinDisplayName = Parameter->GetMetaData(FBlueprintMetadata::MD_DisplayName);
+				if (!PinDisplayName.IsEmpty())
+				{
+					Pin->PinFriendlyName = FText::FromString(PinDisplayName);
+				}
+				// Else cleanup name
+				else
+				{
+					FString PinNameStr = Pin->PinName.ToString();
+					UE::WebAPI::Graph::CleanupPinNameInline(PinNameStr);
+					Pin->PinFriendlyName = FText::FromString(PinNameStr);
+				}
+				
+				//Flag pin as read only for const reference property
+				Pin->bDefaultValueIsIgnored = Parameter->HasAllPropertyFlags(CPF_ConstParm | CPF_ReferenceParm) && (!Function->HasMetaData(FBlueprintMetadata::MD_AutoCreateRefTerm) || Pin->PinType.IsContainer());
+
+				const bool bAdvancedPin = Parameter->HasAllPropertyFlags(CPF_AdvancedDisplay);
+				Pin->bAdvancedView = bAdvancedPin;
+				if(bAdvancedPin && (ENodeAdvancedPins::NoPins == AdvancedPinDisplay))
+				{
+					AdvancedPinDisplay = ENodeAdvancedPins::Hidden;
+				}
+
+				FString ParameterValue;
+				if (K2Schema->FindFunctionParameterDefaultValue(Function, Parameter, ParameterValue))
+				{
+					K2Schema->SetPinAutogeneratedDefaultValue(Pin, ParameterValue);
+				}
+				else
+				{
+					K2Schema->SetPinAutogeneratedDefaultValueBasedOnType(Pin);
+				}
+
+				if (PinsToHide.Contains(Pin->PinName))
+				{
+					Pin->bHidden = true;
+				}
+			}
+
+			UE::WebAPI::Graph::SplitPin(Pin);
+
+			bAllPinsGood = bAllPinsGood && bPinGood;
+		}
+	}
+	
 	Super::AllocateDefaultPins();
 }
 
 FText UK2Node_WebAPIOperation::GetNodeTitle(ENodeTitleType::Type TitleType) const
 {
-	FText FunctionName = Super::GetNodeTitle(TitleType);
+	const FText FunctionName = UK2Node_CallFunction::GetUserFacingFunctionName(GetFactoryFunction());
 	FText NamespaceName;
 	FText ServiceName;
 	
@@ -81,12 +325,86 @@ FText UK2Node_WebAPIOperation::GetNodeTitle(ENodeTitleType::Type TitleType) cons
 	}
 }
 
+FText UK2Node_WebAPIOperation::GetTooltipText() const
+{
+	FText Tooltip;
+
+	UFunction* Function = GetFactoryFunction();
+	if (Function == nullptr)
+	{
+		return FText::Format(LOCTEXT("CallUnknownFunction", "Call unknown function {0}"), FText::FromName(OperationAsyncType == EWebAPIOperationAsyncType::LatentAction ? LatentFunctionName : DelegatedFunctionName));
+	}
+	else if (CachedTooltip.IsOutOfDate(this))
+	{
+		FText BaseTooltip = FText::FromString(UK2Node_CallFunction::GetDefaultTooltipForFunction(Function));
+
+		FFormatNamedArguments Args;
+		Args.Add(TEXT("DefaultTooltip"), BaseTooltip);
+
+		if (Function->HasAllFunctionFlags(FUNC_BlueprintAuthorityOnly))
+		{
+			Args.Add(
+				TEXT("ClientString"),
+				NSLOCTEXT("K2Node", "ServerFunction", "Authority Only. This function will only execute on the server.")
+			);
+			// FText::Format() is slow, so we cache this to save on performance
+			CachedTooltip.SetCachedText(FText::Format(LOCTEXT("WebAPIOperation_SubtitledTooltip", "{DefaultTooltip}\n\n{ClientString}"), Args), this);
+		}
+		else if (Function->HasAllFunctionFlags(FUNC_BlueprintCosmetic))
+		{
+			Args.Add(
+				TEXT("ClientString"),
+				NSLOCTEXT("K2Node", "ClientFunction", "Cosmetic. This event is only for cosmetic, non-gameplay actions.")
+			);
+			// FText::Format() is slow, so we cache this to save on performance
+			CachedTooltip.SetCachedText(FText::Format(LOCTEXT("WebAPIOperation_SubtitledTooltip", "{DefaultTooltip}\n\n{ClientString}"), Args), this);
+		} 
+		else
+		{
+			CachedTooltip.SetCachedText(BaseTooltip, this);
+		}
+	}
+	return CachedTooltip;
+}
+
+void UK2Node_WebAPIOperation::ReconstructNode()
+{
+	if(IsValid())
+	{
+		// Latent type, so remove input delegate pins
+		if(OperationAsyncType == EWebAPIOperationAsyncType::LatentAction)
+		{
+			for(UEdGraphPin* Pin : GetRequestDelegatePins())
+			{
+				RemovePin(Pin);
+			}
+		}
+		// Callback type, so remove (additional) exec pins and response pins
+		else
+		{
+			for(UEdGraphPin* Pin : GetResponseExecPins())
+			{
+				RemovePin(Pin);
+			}
+		}
+	}
+
+	Super::ReconstructNode();
+}
+
+bool UK2Node_WebAPIOperation::IsCompatibleWithGraph(const UEdGraph* TargetGraph) const
+{
+	return UK2Node::IsCompatibleWithGraph(TargetGraph);
+}
+
 void UK2Node_WebAPIOperation::GetNodeContextMenuActions(UToolMenu* Menu, UGraphNodeContextMenuContext* Context) const
 {
 	Super::GetNodeContextMenuActions(Menu, Context);
 
 	if (!Context->bIsDebugging)
 	{
+		// No conversion options if in a function graph (must always be callbacks!)
+		if(GetSchema()->GetGraphType(GetGraph()) != EGraphType::GT_Function)
 		{
 			FText MenuEntryTitle = LOCTEXT("MakeCallbackTitle", "Convert to Callback responses");
 			FText MenuEntryTooltip = LOCTEXT("MakeCallbackTooltip", "Removes the execution pins and instead adds callbacks.");
@@ -105,7 +423,7 @@ void UK2Node_WebAPIOperation::GetNodeContextMenuActions(UToolMenu* Menu, UGraphN
 				const UEdGraphSchema_K2* K2Schema = Cast<UEdGraphSchema_K2>(GetSchema());
 				check(K2Schema != nullptr);
 
-				bCanToggleAsyncType = K2Schema->DoesGraphSupportImpureFunctions(GetGraph()); // @todo 
+				bCanToggleAsyncType = K2Schema->DoesGraphSupportImpureFunctions(GetGraph());
 				if (!bCanToggleAsyncType)
 				{
 					MenuEntryTooltip = LOCTEXT("CannotMakeLatentActionTooltip", "This graph does not support latent actions.");
@@ -130,22 +448,174 @@ void UK2Node_WebAPIOperation::GetNodeContextMenuActions(UToolMenu* Menu, UGraphN
 
 bool UK2Node_WebAPIOperation::CanPasteHere(const UEdGraph* TargetGraph) const
 {
-	// @todo: unlike default async node, you can paste into a function (not event graph) and this should convert to callback style
 	return Super::CanPasteHere(TargetGraph);
 }
 
+void UK2Node_WebAPIOperation::ExpandNode(
+	FKismetCompilerContext& CompilerContext,
+	UEdGraph* SourceGraph)
+{
+	Super::ExpandNode(CompilerContext, SourceGraph);
+	
+	const UEdGraphSchema_K2* Schema = CompilerContext.GetSchema();
+	check(SourceGraph && Schema);
+	bool bIsErrorFree = true;
+
+	const UFunction* FactoryFunction = GetFactoryFunction();
+	if (FactoryFunction == nullptr)
+	{
+		const FName ProxyFunctionName =
+			OperationAsyncType == EWebAPIOperationAsyncType::LatentAction
+			? LatentFunctionName
+			: DelegatedFunctionName;
+			
+		const FText ClassName = OperationClass ? FText::FromString(OperationClass->GetName()) : LOCTEXT("MissingClassString", "Unknown Class");
+		const FString FormattedMessage = FText::Format(
+			LOCTEXT("WebOperationErrorFmt", "WebAPIOperation: Missing function {0} from class {1} for operation @@"),
+			FText::FromString(ProxyFunctionName.GetPlainNameString()),
+			ClassName
+		).ToString();
+
+		CompilerContext.MessageLog.Error(*FormattedMessage, this);
+		return;
+	}
+		
+	// Create a call to factory the proxy object
+	UK2Node_CallFunction* CallCreateOperationNode = CompilerContext.SpawnIntermediateNode<UK2Node_CallFunction>(this, SourceGraph);
+	CallCreateOperationNode->SetFromFunction(FactoryFunction);
+	CallCreateOperationNode->AllocateDefaultPins();
+
+	bIsErrorFree &= CompilerContext.MovePinLinksToIntermediate(*FindPinChecked(UEdGraphSchema_K2::PN_Execute), *CallCreateOperationNode->FindPinChecked(UEdGraphSchema_K2::PN_Execute)).CanSafeConnect();
+
+	// Input pins
+	for (UEdGraphPin* Pin : Pins)
+	{
+		if (UE::WebAPI::Private::ValidDataPin(Pin, EGPD_Input))
+		{
+			UEdGraphPin* DestPin = CallCreateOperationNode->FindPin(Pin->PinName); // match function inputs, to pass data to function from CallFunction node
+			bIsErrorFree &= DestPin && CompilerContext.MovePinLinksToIntermediate(*Pin, *DestPin).CanSafeConnect();
+		}
+	}
+
+	{
+		UEdGraphPin* const OperationObjectPin = CallCreateOperationNode->GetReturnValuePin();
+		check(OperationObjectPin);
+
+		static const FName AsyncTaskProxyName(TEXT("AsyncTaskProxy"));
+		UEdGraphPin* OutputAsyncTaskProxy = FindPin(AsyncTaskProxyName);
+		bIsErrorFree &= !OutputAsyncTaskProxy || CompilerContext.MovePinLinksToIntermediate(*OutputAsyncTaskProxy, *OperationObjectPin).CanSafeConnect();
+
+		bIsErrorFree &= ExpandDefaultToSelfPin(CompilerContext, SourceGraph, CallCreateOperationNode);		
+
+		// GATHER OUTPUT PARAMETERS AND PAIR THEM WITH LOCAL VARIABLES
+		TArray<UE::WebAPI::Private::FOutputPinAndLocalVariable> VariableOutputs;
+		bool bPassedFactoryOutputs = false;
+		for (UEdGraphPin* CurrentPin : Pins)
+		{
+			if ((OutputAsyncTaskProxy != CurrentPin) && UE::WebAPI::Private::ValidDataPin(CurrentPin, EGPD_Output))
+			{
+				if (!bPassedFactoryOutputs)
+				{
+					UEdGraphPin* DestPin = CallCreateOperationNode->FindPin(CurrentPin->PinName);
+					bIsErrorFree &= DestPin && CompilerContext.MovePinLinksToIntermediate(*CurrentPin, *DestPin).CanSafeConnect();
+				}
+				else
+				{
+					const FEdGraphPinType& PinType = CurrentPin->PinType;
+					UK2Node_TemporaryVariable* TempVarOutput = CompilerContext.SpawnInternalVariable(
+						this, PinType.PinCategory, PinType.PinSubCategory, PinType.PinSubCategoryObject.Get(), PinType.ContainerType, PinType.PinValueType);
+					bIsErrorFree &= TempVarOutput->GetVariablePin() && CompilerContext.MovePinLinksToIntermediate(*CurrentPin, *TempVarOutput->GetVariablePin()).CanSafeConnect();
+					VariableOutputs.Add(UE::WebAPI::Private::FOutputPinAndLocalVariable(CurrentPin, TempVarOutput));
+				}
+			}
+			else if (!bPassedFactoryOutputs && CurrentPin && CurrentPin->Direction == EGPD_Output)
+			{
+				// the first exec that isn't the node's then pin is the start of the asyc delegate pins
+				// once we hit this point, we've iterated beyond all outputs for the factory function
+				bPassedFactoryOutputs = (CurrentPin->PinType.PinCategory == UEdGraphSchema_K2::PC_Exec) && (CurrentPin->PinName != UEdGraphSchema_K2::PN_Then);
+			}
+		}
+
+		// FOR EACH DELEGATE DEFINE EVENT, CONNECT IT TO DELEGATE AND IMPLEMENT A CHAIN OF ASSIGMENTS
+		UEdGraphPin* LastThenPin = CallCreateOperationNode->FindPinChecked(UEdGraphSchema_K2::PN_Then);
+
+		UK2Node_CallFunction* IsValidFuncNode = CompilerContext.SpawnIntermediateNode<UK2Node_CallFunction>(this, SourceGraph);
+		const FName IsValidFuncName = GET_FUNCTION_NAME_CHECKED(UKismetSystemLibrary, IsValid);
+		IsValidFuncNode->FunctionReference.SetExternalMember(IsValidFuncName, UKismetSystemLibrary::StaticClass());
+		IsValidFuncNode->AllocateDefaultPins();
+		UEdGraphPin* IsValidInputPin = IsValidFuncNode->FindPinChecked(TEXT("Object"));
+
+		bIsErrorFree &= Schema->TryCreateConnection(OperationObjectPin, IsValidInputPin);
+
+		UK2Node_IfThenElse* ValidateProxyNode = CompilerContext.SpawnIntermediateNode<UK2Node_IfThenElse>(this, SourceGraph);
+		ValidateProxyNode->AllocateDefaultPins();
+		bIsErrorFree &= Schema->TryCreateConnection(IsValidFuncNode->GetReturnValuePin(), ValidateProxyNode->GetConditionPin());
+
+		bIsErrorFree &= Schema->TryCreateConnection(LastThenPin, ValidateProxyNode->GetExecPin());
+		LastThenPin = ValidateProxyNode->GetThenPin();
+
+		if(OperationAsyncType == EWebAPIOperationAsyncType::LatentAction)
+		{
+			bIsErrorFree &= HandleDelegates(VariableOutputs, OperationObjectPin, LastThenPin, SourceGraph, CompilerContext);
+		}
+
+		if (CallCreateOperationNode->FindPinChecked(UEdGraphSchema_K2::PN_Then) == LastThenPin)
+		{
+			CompilerContext.MessageLog.Error(*LOCTEXT("MissingDelegateProperties", "WebAPIOperation: Proxy has no delegates defined. @@").ToString(), this);
+			return;
+		}
+
+		// Move the connections from the original node then pin to the last internal then pin
+
+		UEdGraphPin* OriginalThenPin = FindPin(UEdGraphSchema_K2::PN_Then);
+
+		if (OriginalThenPin)
+		{
+			bIsErrorFree &= CompilerContext.MovePinLinksToIntermediate(*OriginalThenPin, *LastThenPin).CanSafeConnect();
+		}
+		bIsErrorFree &= CompilerContext.CopyPinLinksToIntermediate(*LastThenPin, *ValidateProxyNode->GetElsePin()).CanSafeConnect();
+	}
+
+	if (!bIsErrorFree)
+	{
+		CompilerContext.MessageLog.Error(*LOCTEXT("InternalConnectionError", "WebAPIOperation: Internal connection error. @@").ToString(), this);
+	}
+
+	// Make sure we caught everything
+	BreakAllNodeLinks();
+}
+
+FName UK2Node_WebAPIOperation::GetCornerIcon() const
+{
+	// Only return the latent action icon if that's the current async type, otherwise empty
+	return OperationAsyncType == EWebAPIOperationAsyncType::LatentAction
+		? Super::GetCornerIcon()
+		: NAME_None;
+}
+
 void UK2Node_WebAPIOperation::PostPlacedNewNode()
 {
+	// Prevent latent function in function graph
+	if(OperationAsyncType == EWebAPIOperationAsyncType::LatentAction && GetSchema()->GetGraphType(GetGraph()) == GT_Function)
+	{
+		OperationAsyncType = EWebAPIOperationAsyncType::Callback;
+		ReconstructNode();
+	}
+	
 	Super::PostPlacedNewNode();
-
-	PostReconstructNode();
 }
 
 void UK2Node_WebAPIOperation::PostPasteNode()
 {
-	Super::PostPasteNode();
+	const bool bIsFunctionGraph = GetSchema()->GetGraphType(GetGraph()) == GT_Function;
+	
+	// Latent action pasted in function graph, forcibly convert to callback style
+	if(OperationAsyncType == EWebAPIOperationAsyncType::LatentAction && bIsFunctionGraph)
+	{
+		ToggleAsyncType();
+	}
 
-	// @todo: handle conversion to/from callback style
+	Super::PostPasteNode();
 }
 
 void UK2Node_WebAPIOperation::GetMenuActions(FBlueprintActionDatabaseRegistrar& ActionRegistrar) const
@@ -154,44 +624,59 @@ void UK2Node_WebAPIOperation::GetMenuActions(FBlueprintActionDatabaseRegistrar&
 	{
 		static void SetNodeFunc(UEdGraphNode* NewNode, bool /*bIsTemplateNode*/, TWeakObjectPtr<UFunction> FunctionPtr)
 		{
-			UK2Node_WebAPIOperation* AsyncTaskNode = CastChecked<UK2Node_WebAPIOperation>(NewNode);
+			UK2Node_WebAPIOperation* OperationNode = CastChecked<UK2Node_WebAPIOperation>(NewNode);
 			if (FunctionPtr.IsValid())
 			{
 				const UFunction* Func = FunctionPtr.Get();
-				FObjectProperty* ReturnProp = CastFieldChecked<FObjectProperty>(Func->GetReturnProperty());
 
-				AsyncTaskNode->ProxyFactoryFunctionName = Func->GetFName();
-				AsyncTaskNode->ProxyFactoryClass        = Func->GetOuterUClass();
-				AsyncTaskNode->ProxyClass               = ReturnProp->PropertyClass;				
+				FString LatentFunctionName = Func->GetName();
+				FString DelegatedFunctionName = LatentFunctionName;
+				DelegatedFunctionName.RemoveFromEnd(TEXT("Async"));				
+
+				OperationNode->OperationClass			= Func->GetOuterUClass();
+				OperationNode->LatentFunctionName		= FName(LatentFunctionName);
+				OperationNode->DelegatedFunctionName	= FName(DelegatedFunctionName);
+
+				check(OperationNode->CacheOutcomeDelegates());
 			}
 		}
 	};
 
 	UClass* NodeClass = GetClass();
-	ActionRegistrar.RegisterClassFactoryActions<UWebAPIOperationObject>(FBlueprintActionDatabaseRegistrar::FMakeFuncSpawnerDelegate::CreateLambda([NodeClass](const UFunction* FactoryFunc)->UBlueprintNodeSpawner*
-	{
-		UBlueprintNodeSpawner* NodeSpawner = UBlueprintFunctionNodeSpawner::Create(FactoryFunc);
-		check(NodeSpawner != nullptr);
-		
-		NodeSpawner->NodeClass = NodeClass;
+	ActionRegistrar.RegisterClassFactoryActions<UWebAPIOperationObject>(
+		FBlueprintActionDatabaseRegistrar::FMakeFuncSpawnerDelegate::CreateLambda(
+			[NodeClass](const UFunction* FactoryFunc) -> UBlueprintNodeSpawner*
+			{
+				// Skip functions without "Async" suffix - these should explicitly not be shown.
+				// switching implementations between latent (async suffix) and callback (no suffix) is automatic.
+				if(!FactoryFunc->GetName().EndsWith(TEXT("Async")))
+				{
+					return nullptr;
+				}
+
+				UBlueprintNodeSpawner* NodeSpawner = UBlueprintFunctionNodeSpawner::Create(FactoryFunc);
+				check(NodeSpawner != nullptr);
+
+				NodeSpawner->NodeClass = NodeClass;
 
-		const TWeakObjectPtr<UFunction> FunctionPtr = MakeWeakObjectPtr(const_cast<UFunction*>(FactoryFunc));
-		NodeSpawner->CustomizeNodeDelegate = UBlueprintNodeSpawner::FCustomizeNodeDelegate::CreateStatic(GetMenuActions_Utils::SetNodeFunc, FunctionPtr);
+				const TWeakObjectPtr<UFunction> FunctionPtr = MakeWeakObjectPtr(const_cast<UFunction*>(FactoryFunc));
+				NodeSpawner->CustomizeNodeDelegate = UBlueprintNodeSpawner::FCustomizeNodeDelegate::CreateStatic(GetMenuActions_Utils::SetNodeFunc, FunctionPtr);
 
-		return NodeSpawner;
-	}));
+				return NodeSpawner;
+			}));
 }
 
 void UK2Node_WebAPIOperation::PostReconstructNode()
 {
-	Super::PostReconstructNode();
-
-	if(IsValid())
+	// Prevent latent function in function graph
+	if(OperationAsyncType == EWebAPIOperationAsyncType::LatentAction && GetSchema()->GetGraphType(GetGraph()) == GT_Function)
 	{
-		SplitRequestPins(GetRequestPins());
-		SplitResponsePins(GetResponsePins());
-		SplitResponsePins(GetErrorResponsePins());
+		ToggleAsyncType();
+		return;
 	}
+
+	Super::PostReconstructNode();
+	InvalidatePinTooltips();
 }
 
 void UK2Node_WebAPIOperation::SetAsyncType(EWebAPIOperationAsyncType InAsyncType)
@@ -200,6 +685,16 @@ void UK2Node_WebAPIOperation::SetAsyncType(EWebAPIOperationAsyncType InAsyncType
 	{
 		OperationAsyncType = InAsyncType;
 
+		// Provides opportunity to perform some actions specific to the conversion direction
+		if(OperationAsyncType == EWebAPIOperationAsyncType::LatentAction)
+		{
+			ConvertCallbackToLatent();			
+		}
+		else
+		{
+			ConvertLatentToCallback();
+		}
+
 		const bool bHasBeenConstructed = (Pins.Num() > 0);
 		if (bHasBeenConstructed)
 		{
@@ -210,151 +705,507 @@ void UK2Node_WebAPIOperation::SetAsyncType(EWebAPIOperationAsyncType InAsyncType
 
 bool UK2Node_WebAPIOperation::IsValid() const
 {
-	return ProxyFactoryFunctionName != NAME_None
-		&& ProxyClass != nullptr
+	return (LatentFunctionName != NAME_None || DelegatedFunctionName != NAME_None)
+		&& OperationClass != nullptr
 		&& ((bHasCompilerMessage && ErrorType >= EMessageSeverity::Info) || !bHasCompilerMessage);
 }
 
-UEdGraphPin* UK2Node_WebAPIOperation::FindPinContainingName(
-	const EEdGraphPinDirection& InPinDirection,
-	const FString& InName) const
+const UK2Node_CustomEvent* UK2Node_WebAPIOperation::GetCustomEventForOutcomeDelegate(
+	const UK2Node_WebAPIOperation* InNode,
+	const FName& InOutcomeName)
 {
-	return FindPinByPredicate([&InPinDirection, &InName](const UEdGraphPin* InPin)
+	check(InNode);
+
+	TArray<UEdGraphPin*> DelegatePins = InNode->GetRequestDelegatePins();
+	if(UEdGraphPin** OutcomeDelegatePin = DelegatePins.FindByPredicate(
+		[InOutcomeName](const UEdGraphPin* InPin)
 	{
-		return InPin->Direction == InPinDirection
-			&& InPin->PinType.PinCategory != UEdGraphSchema_K2::PC_Exec
-			&& InPin->GetName().Contains(InName);
-	});
+			return InPin->PinName.ToString().Contains(InOutcomeName.ToString());
+	}))
+	{
+		// Found pin, now get connected event
+		for(const UEdGraphPin* LinkedPin : (*OutcomeDelegatePin)->LinkedTo)
+		{
+			if(LinkedPin->GetOwningNode()->IsA<UK2Node_CustomEvent>())
+			{
+				return Cast<UK2Node_CustomEvent>(LinkedPin->GetOwningNode());
+			}
+		}
+	}
+
+	return nullptr;
 }
 
-TArray<UEdGraphPin*> UK2Node_WebAPIOperation::FindPinsContainingName(
-	const EEdGraphPinDirection& InPinDirection,
-	const FString& InName) const
+UFunction* UK2Node_WebAPIOperation::GetFactoryFunction() const
 {
-	return Pins.FilterByPredicate([&InPinDirection, InName](const UEdGraphPin* InPin)
+	if (OperationClass == nullptr)
 	{
-		return InPin->Direction == InPinDirection
-			&& InPin->PinType.PinCategory != UEdGraphSchema_K2::PC_Exec
-			&& InPin->GetName().Contains(InName);
-	});
+		UE_LOG(LogWebAPIBlueprintGraph, Error, TEXT("OperationClass null in %s. Was a class deleted or saved on a non promoted build?"), *GetFullName());
+		return nullptr;
+	}
+
+	const FName ProxyFunctionName =
+		OperationAsyncType == EWebAPIOperationAsyncType::LatentAction
+		? LatentFunctionName
+		: DelegatedFunctionName;
+	
+	FMemberReference FunctionReference;
+	FunctionReference.SetExternalMember(ProxyFunctionName, OperationClass);
+
+	UFunction* FactoryFunction = FunctionReference.ResolveMember<UFunction>(GetBlueprint());
+	
+	if (FactoryFunction == nullptr)
+	{
+		FactoryFunction = OperationClass->FindFunctionByName(ProxyFunctionName);
+		UE_CLOG(FactoryFunction == nullptr, LogWebAPIBlueprintGraph, Error, TEXT("FactoryFunction %s null in %s. Was a class deleted or saved on a non promoted build?"), *ProxyFunctionName.ToString(), *GetFullName());
+	}
+
+	return FactoryFunction;
+}
+
+void UK2Node_WebAPIOperation::GetRedirectPinNames(
+	const UEdGraphPin& Pin,
+	TArray<FString>& RedirectPinNames) const
+{
+	Super::GetRedirectPinNames(Pin, RedirectPinNames);
+}
+
+bool UK2Node_WebAPIOperation::ExpandDefaultToSelfPin(
+	FKismetCompilerContext& InCompilerContext,
+	UEdGraph* InSourceGraph,
+	UK2Node_CallFunction* InIntermediateProxyNode)
+{
+	if(InSourceGraph && InIntermediateProxyNode)
+	{
+		// Connect a self reference pin if there is a TScriptInterface default to self
+		if (const UFunction* TargetFunc = InIntermediateProxyNode->GetTargetFunction())
+		{
+			const FString& MetaData = TargetFunc->GetMetaData(FBlueprintMetadata::MD_DefaultToSelf);
+			if (!MetaData.IsEmpty())
+			{
+				// Find the default to self value pin
+				if (UEdGraphPin* DefaultToSelfPin = InIntermediateProxyNode->FindPinChecked(MetaData, EGPD_Input))
+				{
+					// If it has no links then spawn a new self node here
+					if (DefaultToSelfPin->PinType.PinCategory == UEdGraphSchema_K2::PC_Interface && DefaultToSelfPin->LinkedTo.Num() == 0)
+					{
+						const UEdGraphSchema_K2* Schema = InCompilerContext.GetSchema();
+
+						UK2Node_Self* SelfNode = InCompilerContext.SpawnIntermediateNode<UK2Node_Self>(this, InSourceGraph);
+						SelfNode->AllocateDefaultPins();
+						UEdGraphPin* SelfPin = SelfNode->FindPinChecked(UEdGraphSchema_K2::PSC_Self);
+						// Make a connection from this intermediate self pin to here
+						return Schema->TryCreateConnection(DefaultToSelfPin, SelfPin);
+					}
+				}
+			}
+		}	
+	}
+	return true;
+}
+
+bool UK2Node_WebAPIOperation::HandleDelegates(
+	const TArray<UE::WebAPI::Private::FOutputPinAndLocalVariable>& VariableOutputs,
+	UEdGraphPin* ProxyObjectPin,
+	UEdGraphPin*& InOutLastThenPin,
+	UEdGraph* SourceGraph,
+	FKismetCompilerContext& CompilerContext)
+{
+	bool bIsErrorFree = true;
+	for (TFieldIterator<FMulticastDelegateProperty> PropertyIt(OperationClass); PropertyIt && bIsErrorFree; ++PropertyIt)
+	{
+		UEdGraphPin* LastActivatedThenPin = nullptr;
+		bIsErrorFree &= UE::WebAPI::Private::HandleDelegateImplementation(*PropertyIt, VariableOutputs, ProxyObjectPin, InOutLastThenPin, LastActivatedThenPin, this, SourceGraph, CompilerContext);
+	}
+	return bIsErrorFree;
+}
+
+bool UK2Node_WebAPIOperation::CacheOutcomeDelegates()
+{
+	// If either invalid
+	if(!PositiveDelegateProperty.IsValid() || !NegativeDelegateProperty.IsValid())
+	{
+		check(OperationClass);
+
+		static FString PositiveOutcomeNameStr = UE::WebAPI::Operation::PositiveOutcomeName.ToString();
+		static FString NegativeOutcomeNameStr = UE::WebAPI::Operation::NegativeOutcomeName.ToString();
+		
+		for (TFieldIterator<FMulticastDelegateProperty> PropertyIt(OperationClass); PropertyIt; ++PropertyIt)
+		{
+			if((*PropertyIt)->GetName().Contains(PositiveOutcomeNameStr))
+			{
+				PositiveDelegateProperty = *PropertyIt;
+			}
+			else if((*PropertyIt)->GetName().Contains(NegativeOutcomeNameStr))
+			{
+				NegativeDelegateProperty = *PropertyIt;
+			}
+		}	
+	}
+
+	check(PositiveDelegateProperty.IsValid());
+	check(NegativeDelegateProperty.IsValid());
+
+	return true;
+}
+
+void UK2Node_WebAPIOperation::InvalidatePinTooltips()
+{
+	bPinTooltipsValid = false;
+}
+
+UEdGraphPin* UK2Node_WebAPIOperation::FindPin(
+	const FName& InName,
+	const EEdGraphPinDirection& InDirection,
+	const FName& InCategory,
+	bool bFindPartial) const
+{
+	return UE::WebAPI::Graph::FindPin(this, InName, InDirection, InCategory, bFindPartial);	
+}
+
+TArray<UEdGraphPin*> UK2Node_WebAPIOperation::FindPins(
+	const FString& InName,
+	const EEdGraphPinDirection& InDirection,
+	bool bOnlySplitPins) const
+{
+	return UE::WebAPI::Graph::FindPins(this, InName, InDirection, bOnlySplitPins);
 }
 
 TArray<UEdGraphPin*> UK2Node_WebAPIOperation::GetRequestPins() const
 {
-	TArray<UEdGraphPin*> RequestPins = FindPinsContainingName(EGPD_Input, TEXT("Request"));
+	TArray<UEdGraphPin*> RequestPins = FindPins(TEXT("Request"), EGPD_Input);
 	ensureMsgf(!RequestPins.IsEmpty(), TEXT("The function must contain a parameter with \"Request\" in the name."));	
 	return RequestPins;
 }
 
+TArray<UEdGraphPin*> UK2Node_WebAPIOperation::GetRequestDelegatePins() const
+{
+	return Pins.FilterByPredicate([](const UEdGraphPin* InPin)
+	{
+		return InPin->Direction == EGPD_Input
+			&& InPin->PinType.PinCategory == UEdGraphSchema_K2::PC_Delegate;
+	});
+}
+
 TArray<UEdGraphPin*> UK2Node_WebAPIOperation::GetResponsePins() const
 {
-	auto ResponsePins = FindPinsContainingName(EGPD_Output, TEXT("Response"));
-	ensureMsgf(!ResponsePins.IsEmpty(), TEXT("The Operation must contain a delegate with a parameter containing \"Response\" in the name."));	
-	return ResponsePins;
+	return UE::WebAPI::Graph::GetResponsePins(this);
+}
+
+UEdGraphPin* UK2Node_WebAPIOperation::GetThenPin() const
+{
+	return FindPin(UEdGraphSchema_K2::PN_Then, EEdGraphPinDirection::EGPD_Output, UEdGraphSchema_K2::PC_Exec);
+}
+
+TArray<UEdGraphPin*> UK2Node_WebAPIOperation::GetResponseExecPins() const
+{
+	return Pins.FilterByPredicate([](const UEdGraphPin* InPin)
+	{
+		return InPin->Direction == EGPD_Output
+			&& InPin->PinName.ToString().StartsWith(TEXT("On"))
+			&& InPin->PinType.PinCategory == UEdGraphSchema_K2::PC_Exec;
+	});
 }
 
 TArray<UEdGraphPin*> UK2Node_WebAPIOperation::GetErrorResponsePins() const
-{	
-	auto ErrorResponsePins = FindPinsContainingName(EGPD_Output, TEXT("Error"));
-	ensureMsgf(!ErrorResponsePins.IsEmpty(), TEXT("The Operation must contain a delegate with a parameter containing \"Error\" in the name."));	
-	return ErrorResponsePins;
+{
+	return UE::WebAPI::Graph::GetErrorResponsePins(this);
 }
 
-// @todo: recursively split where struct is basically a wrapper for another struct
-void UK2Node_WebAPIOperation::SplitRequestPins(const TArray<UEdGraphPin*>& InPins) const
+UK2Node_CustomEvent* UK2Node_WebAPIOperation::MakeCustomEvent(
+	const FName& InName,
+	const UFunction* InSignature,
+	const FVector2D& InPosition)
 {
-	// If UWebAPIOperation base class, the pin won't always be valid, so ignore.
-	if(!ensureAlways(!InPins.IsEmpty() && Algo::AllOf(InPins, [](const UEdGraphPin* InPin) { return InPin; })))
+	if(!GetBlueprint() || !GetBlueprint()->SupportsEventGraphs())
 	{
-		return;
+		return nullptr;
 	}
 
-	for(UEdGraphPin* Pin : InPins)
+	UEdGraph* Graph = GetGraph();
+	if(GetSchema()->GetGraphType(GetGraph()) != EGraphType::GT_Ubergraph)
 	{
-		// Split if not already
-		if (Pin->SubPins.Num() == 0)
+		Graph = GetBlueprint()->GetLastEditedUberGraph();
+		if(!Graph)
 		{
-			if(CanSplitPin(Pin))
+			FKismetEditorUtilities::CreateDefaultEventGraphs(GetBlueprint());
+			// Something went wrong, return nullptr...
+			if(GetBlueprint()->UbergraphPages.IsEmpty())
 			{
-				const UEdGraphSchema_K2* K2Schema = GetDefault<UEdGraphSchema_K2>();
-				K2Schema->SplitPin(Pin);
+				return nullptr;
 			}
-		}
-
-		// Removes the prefix from the auto-named split pins 
-		if (Pin->SubPins.Num() > 0)
-		{
-			for (UEdGraphPin* SubPin : Pin->SubPins)
+			// ... otherwise retry after adding event graph
+			else
 			{
-				FString SubPinName = SubPin->PinName.ToString();
-				CleanupPinNameInline(SubPinName);
-
-				//SubPin->Modify();			
-				SubPin->PinFriendlyName = FText::FromString(SubPinName);
+				return MakeCustomEvent(InName, InSignature, InPosition);
 			}
 		}
 	}
+
+	const FString FunctionName = FString::Printf(TEXT("%s_Event"), *InName.ToString());
+	UK2Node_CustomEvent* CustomEventNode = UK2Node_CustomEvent::CreateFromFunction(InPosition, Graph, FunctionName, InSignature, false);
+
+	UE::WebAPI::Graph::SplitPins(UE::WebAPI::Graph::GetResponsePins(CustomEventNode));
+	UE::WebAPI::Graph::SplitPins(UE::WebAPI::Graph::GetErrorResponsePins(CustomEventNode));
+
+	return CustomEventNode;
 }
 
-// Splits if the Response contains a single property
-void UK2Node_WebAPIOperation::SplitResponsePins(const TArray<UEdGraphPin*>& InPins) const
+/** Named references to Pins. */
+struct FLatentActionPinMap
 {
-	// Compilation errors or stale nodes should be warnings, not errors
-	if(!ensureAlways(!InPins.IsEmpty() && Algo::AllOf(InPins, [](const UEdGraphPin* InPin) { return InPin; })))
+public:
+	explicit FLatentActionPinMap(const UK2Node_WebAPIOperation* InNode)
 	{
-		return;	
+		check(InNode);
+
+		ExecIn = InNode->GetExecPin();
+		RequestParameters = InNode->GetRequestPins();
+
+		ExecOut = InNode->GetThenPin();
+		ExecPositive = InNode->FindPin(UE::WebAPI::Operation::PositiveOutcomeName, EEdGraphPinDirection::EGPD_Output, UEdGraphSchema_K2::PC_Exec, true);
+		ExecNegative = InNode->FindPin(UE::WebAPI::Operation::NegativeOutcomeName, EEdGraphPinDirection::EGPD_Output, UEdGraphSchema_K2::PC_Exec, true);
+		Responses = InNode->GetResponsePins();
+		Responses.Append(InNode->GetErrorResponsePins());
 	}
+	
+	// Inputs
+	UEdGraphPin* ExecIn;
 
-	for(UEdGraphPin* Pin : InPins)
+	// (Non-split pins)
+	TArray<UEdGraphPin*> RequestParameters;
+
+	// Outputs
+	UEdGraphPin* ExecOut;
+	UEdGraphPin* ExecPositive;
+	UEdGraphPin* ExecNegative;
+
+	// (Non-split pins)
+	TArray<UEdGraphPin*> Responses;
+
+	UEdGraphPin* GetOutcomeExec(int32 InIndex) const
 	{
-		if (const UScriptStruct* PinStructType = Cast<UScriptStruct>(Pin->PinType.PinSubCategoryObject.Get()))
-		{
-			uint32 PropertyNum = 0;
-			
-			// Split if not already
-			if (Pin->SubPins.Num() == 0 && CanSplitPin(Pin))
-			{
-				for (TFieldIterator<const FProperty> PropertyIt(PinStructType, EFieldIteratorFlags::ExcludeSuper, EFieldIteratorFlags::ExcludeDeprecated); PropertyIt; ++PropertyIt)
-				{
-					if (PropertyIt->HasAnyPropertyFlags(CPF_Transient))
-					{
-						continue;
-					}
+		return InIndex == 0 ? ExecPositive : ExecNegative;
+	}
 
-					PropertyNum++;
-				}
+	bool IsValid() const
+	{
+		return ExecIn
+			&& (RequestParameters.IsEmpty() ? true : Algo::AllOf(RequestParameters))
+			&& ExecOut
+			&& ExecPositive
+			&& ExecNegative
+			&& (Responses.IsEmpty() ? true : Algo::AllOf(Responses));		
+	}
+};
 
-				if(PropertyNum == 1)
-				{
-					// Split, result should be single pin (cause it only contained a single property!)
-					const UEdGraphSchema_K2* K2Schema = GetDefault<UEdGraphSchema_K2>();
-					K2Schema->SplitPin(Pin);
-				}
-			}
+/** Named references to Pins for a single custom event. */
+struct FOutcomeEventPinMap
+{
+public:
+	FOutcomeEventPinMap() = default;
+	explicit FOutcomeEventPinMap(FName InName, const UK2Node_CustomEvent* InNode, const UFunction* InSignature)
+	{
+		check(InNode);
+		check(InSignature);
 
-			// Removes the prefix from the auto-named split pins 
-			if (Pin->SubPins.Num() > 0)
-			{
-				for (UEdGraphPin* SubPin : Pin->SubPins)
-				{
-					FString SubPinName = SubPin->PinName.ToString();
-					CleanupPinNameInline(SubPinName);
+		OutcomeName = InName;
+		ExecOut = UE::WebAPI::Graph::FindPin(InNode, UEdGraphSchema_K2::PN_Then, EEdGraphPinDirection::EGPD_Output, UEdGraphSchema_K2::PC_Exec);
+		DelegatePin =  InNode->FindPinChecked(InNode->DelegateOutputName);
+		Responses = UE::WebAPI::Graph::GetResponsePins(InNode);
+		Responses.Append(UE::WebAPI::Graph::GetErrorResponsePins(InNode));
+	}
 
-					//SubPin->Modify();			
-					SubPin->PinFriendlyName = FText::FromString(SubPinName);
-				}
-			}
-		}
+	FName OutcomeName = NAME_None;
+	
+	// Outputs
+	UEdGraphPin* ExecOut = nullptr;
+
+	UEdGraphPin* DelegatePin = nullptr;
+
+	// (Non-split pins)
+	TArray<UEdGraphPin*> Responses = {};
+
+	bool IsValid() const
+	{
+		return ExecOut
+			&& DelegatePin
+			&& (Responses.IsEmpty() ? true : Algo::AllOf(Responses));		
 	}
+};
+
+/** Named references to Pins. */
+struct FCallbackActionPinMap
+{
+public:
+	/** This overload gets the event nodes from the connected delegate pins. */
+	explicit FCallbackActionPinMap(const UK2Node_WebAPIOperation* InNode)
+		: FCallbackActionPinMap(InNode
+			, UK2Node_WebAPIOperation::GetCustomEventForOutcomeDelegate(InNode, UE::WebAPI::Operation::PositiveOutcomeName)
+			, UK2Node_WebAPIOperation::GetCustomEventForOutcomeDelegate(InNode, UE::WebAPI::Operation::NegativeOutcomeName))
+	{
+	}
+	
+	explicit FCallbackActionPinMap(
+		const UK2Node_WebAPIOperation* InNode,
+		const UK2Node_CustomEvent* InPositiveEvent,
+		const UK2Node_CustomEvent* InNegativeEvent)
+	{
+		check(InNode);
+		check(InPositiveEvent);
+		check(InNegativeEvent);
+		
+		ExecIn = InNode->GetExecPin();
+		RequestParameters = InNode->GetRequestPins();
+		CallbackPositive = InNode->FindPin(UE::WebAPI::Operation::PositiveOutcomeName, EEdGraphPinDirection::EGPD_Input, UEdGraphSchema_K2::PC_Delegate, true);
+		CallbackNegative = InNode->FindPin(UE::WebAPI::Operation::NegativeOutcomeName, EEdGraphPinDirection::EGPD_Input, UEdGraphSchema_K2::PC_Delegate, true);
+
+		check(InNode->PositiveDelegateProperty.IsValid());
+		DelegateProperty = InNode->PositiveDelegateProperty.Get(); // Choose one, doesn't matter
+		
+		ExecOut = InNode->GetThenPin();
+		
+		PositiveEventNodeMap = FOutcomeEventPinMap(UE::WebAPI::Operation::PositiveOutcomeName, InPositiveEvent, DelegateProperty->SignatureFunction);
+		NegativeEventNodeMap = FOutcomeEventPinMap(UE::WebAPI::Operation::NegativeOutcomeName, InNegativeEvent, DelegateProperty->SignatureFunction);
+	}
+	
+	// Inputs
+	UEdGraphPin* ExecIn;
+	
+	// (Non-split pins)
+	TArray<UEdGraphPin*> RequestParameters;
+
+	// Outcome delegate property (positive and negative outcomes have same signature)
+	FMulticastDelegateProperty* DelegateProperty;
+	
+	UEdGraphPin* CallbackPositive;
+	UEdGraphPin* CallbackNegative;
+
+	// Outputs
+	UEdGraphPin* ExecOut;
+
+	// Event Nodes
+	FOutcomeEventPinMap PositiveEventNodeMap;
+	FOutcomeEventPinMap NegativeEventNodeMap;
+
+	UEdGraphPin* GetCallback(int32 InIndex) const
+	{
+		return InIndex == 0 ? CallbackPositive : CallbackNegative;
+	}
+
+	const FOutcomeEventPinMap& GetEventNodeMap(int32 InIndex) const
+	{
+		return InIndex == 0 ? PositiveEventNodeMap : NegativeEventNodeMap;
+	}
+
+	bool IsValid() const
+	{
+		return ExecIn
+			&& (RequestParameters.IsEmpty() ? true : Algo::AllOf(RequestParameters))
+			&& CallbackPositive
+			&& CallbackNegative
+			&& ExecOut;
+	}
+};
+
+void UK2Node_WebAPIOperation::ConvertLatentToCallback()
+{
+	CacheOutcomeDelegates();
+
+	const FLatentActionPinMap LatentActionPinMap(this);
+	check(LatentActionPinMap.IsValid());
+
+	const static TArray<FName, TFixedAllocator<2>> OutcomeNames = {UE::WebAPI::Operation::PositiveOutcomeName, UE::WebAPI::Operation::NegativeOutcomeName };
+	TArray<UEdGraphPin*, TFixedAllocator<2>> LatentActionOutcomePins = { LatentActionPinMap.ExecPositive, LatentActionPinMap.ExecNegative };
+
+	TArray<UK2Node_CustomEvent*, TFixedAllocator<2>> OutcomeEvents = { nullptr, nullptr };
+
+	// @note: First pass requires old pins
+	for(int32 Idx = 0; Idx < OutcomeNames.Num(); ++Idx)
+	{
+		const UEdGraphPin* OutcomeExecPin = LatentActionOutcomePins[Idx];
+		FName OutcomeName = OutcomeNames[Idx];
+
+		// Make custom event
+		FVector2D OperationNodePosition = { static_cast<double>(NodePosX), static_cast<double>(NodePosY) };
+		OperationNodePosition += FVector2D(-200, OutcomeName == UE::WebAPI::Operation::PositiveOutcomeName ? 200 : 400);
+		UK2Node_CustomEvent* CustomEvent = MakeCustomEvent(OutcomeName, PositiveDelegateProperty->SignatureFunction, OperationNodePosition);
+		OutcomeEvents[Idx] = CustomEvent;
+		check(CustomEvent);
+
+		// Wire event exec to linked
+		UEdGraphPin* EventExecPin = UE::WebAPI::Graph::FindPin(CustomEvent, UEdGraphSchema_K2::PN_Then, EEdGraphPinDirection::EGPD_Output, UEdGraphSchema_K2::PC_Exec);
+		UE::WebAPI::Graph::TransferPinConnections(OutcomeExecPin, EventExecPin);
+	}
+
+	// Refresh
+	AllocateDefaultPins();
+
+	const FCallbackActionPinMap CallbackActionPinMap(this, OutcomeEvents[0], OutcomeEvents[1]);
+	check(CallbackActionPinMap.IsValid());
+
+	// @note: Second pass requires new pins
+	for(int32 Idx = 0; Idx < OutcomeNames.Num(); ++Idx)
+	{
+		UEdGraphPin* CustomEventDelegatePin = CallbackActionPinMap.GetEventNodeMap(Idx).DelegatePin;
+
+		// Exec pin
+		UEdGraphPin* OperationDelegatePin = CallbackActionPinMap.GetCallback(Idx);
+		OperationDelegatePin->GetSchema()->TryCreateConnection(CustomEventDelegatePin, OperationDelegatePin);
+
+		TArray<UEdGraphPin*> PinsInUse = UE::WebAPI::Graph::FilterPinsByRelated(LatentActionOutcomePins[Idx], LatentActionPinMap.Responses);
+		UE::WebAPI::Graph::TransferPins(PinsInUse, CallbackActionPinMap.GetEventNodeMap(Idx).Responses);
+	}
+	
+	// Refresh
+	ReconstructNode();
 }
 
-void UK2Node_WebAPIOperation::CleanupPinNameInline(FString& InPinName)
+void UK2Node_WebAPIOperation::ConvertCallbackToLatent()
 {
-	int32 UnderscorePosition = -1;
-	if(InPinName.FindLastChar(TCHAR('_'), UnderscorePosition))
+	CacheOutcomeDelegates();
+
+	const static TArray<FName, TFixedAllocator<2>> OutcomeNames = { UE::WebAPI::Operation::PositiveOutcomeName, UE::WebAPI::Operation::NegativeOutcomeName };
+
+	// @note: overloaded ctor finds events from connected delegate pins
+	const FCallbackActionPinMap CallbackActionPinMap(this);
+	check(CallbackActionPinMap.IsValid());
+
+	// Refresh
+	AllocateDefaultPins();
+
+	const FLatentActionPinMap LatentActionPinMap(this);
+	check(LatentActionPinMap.IsValid());
+
+	for(int32 Idx = 0; Idx < OutcomeNames.Num(); ++Idx)
 	{
-		InPinName = InPinName.RightChop(UnderscorePosition + 1);
+		// Get custom event for this outcome (from previous callback setup)
+		const FOutcomeEventPinMap& CustomEventPinMap = CallbackActionPinMap.GetEventNodeMap(Idx);
+
+		// Wire custom event exec to latent exec pin
+		const UEdGraphPin* EventExecPin = CustomEventPinMap.ExecOut;
+		UE::WebAPI::Graph::TransferPinConnections(EventExecPin, LatentActionPinMap.GetOutcomeExec(Idx));
+		
+		// Get custom event exec pin and find relevant nodes (
+		TArray<UEdGraphPin*> PinsInUse = UE::WebAPI::Graph::FilterPinsByRelated(CustomEventPinMap.ExecOut, CustomEventPinMap.Responses);
+		UE::WebAPI::Graph::TransferPins(PinsInUse, LatentActionPinMap.Responses);
+
+		const bool bNothingLinkedToCustomEvent = Algo::AllOf(CustomEventPinMap.Responses,
+			[](const UEdGraphPin* InPin)
+			{
+				return InPin->LinkedTo.IsEmpty();
+			});
+
+		// If nothing linked, delete custom event
+		if(bNothingLinkedToCustomEvent)
+		{
+			UEdGraphNode* CustomEventNode = CustomEventPinMap.ExecOut->GetOwningNode();
+			CustomEventNode->GetGraph()->RemoveNode(CustomEventNode);			
+		}
 	}
+
+	// Refresh
+	ReconstructNode();
 }
 
 void UK2Node_WebAPIOperation::ToggleAsyncType()
@@ -369,11 +1220,4 @@ void UK2Node_WebAPIOperation::ToggleAsyncType()
 	SetAsyncType(StaticCast<EWebAPIOperationAsyncType>(FMath::Abs(StaticCast<int8>(OperationAsyncType) - 1)));
 }
 
-bool UK2Node_WebAPIOperation::ReconnectExecPins(TArray<UEdGraphPin*>& OldPins)
-{
-	// @todo: handle when converting to/from callback style
-	
-	return true;
-}
-
 #undef LOCTEXT_NAMESPACE
diff --git a/Engine/Plugins/Web/WebAPI/Source/WebAPIBlueprintGraph/Private/K2Node_WebAPIOperation.h b/Engine/Plugins/Web/WebAPI/Source/WebAPIBlueprintGraph/Private/K2Node_WebAPIOperation.h
--- a/Engine/Plugins/Web/WebAPI/Source/WebAPIBlueprintGraph/Private/K2Node_WebAPIOperation.h
+++ b/Engine/Plugins/Web/WebAPI/Source/WebAPIBlueprintGraph/Private/K2Node_WebAPIOperation.h
@@ -3,12 +3,30 @@
 #pragma once
 
 #include "CoreMinimal.h"
-#include "K2Node_BaseAsyncTask.h"
+#include "EdGraph/EdGraphNodeUtils.h"
+#include "K2Node.h"
+#include "UObject/WeakFieldPtr.h"
 
 #include "K2Node_WebAPIOperation.generated.h"
 
+class UK2Node_CallFunction;
+class UK2Node_TemporaryVariable;
 class FBlueprintActionDatabaseRegistrar;
 
+namespace UE::WebAPI::Private
+{
+	struct FOutputPinAndLocalVariable
+	{
+		UEdGraphPin* OutputPin;
+		UK2Node_TemporaryVariable* TempVar;
+
+		FOutputPinAndLocalVariable(UEdGraphPin* Pin, UK2Node_TemporaryVariable* Var)
+			: OutputPin(Pin)
+			, TempVar(Var) { }
+	};
+}
+
+
 UENUM()
 enum class EWebAPIOperationAsyncType : uint8
 {
@@ -21,16 +39,17 @@ enum class EWebAPIOperationAsyncType : uint8
  */
 UCLASS()
 class WEBAPIBLUEPRINTGRAPH_API UK2Node_WebAPIOperation
-	: public UK2Node_BaseAsyncTask
+	: public UK2Node
 {
 	GENERATED_BODY()
 
 public:
-	UK2Node_WebAPIOperation();
-
 	//~ Begin UEdGraphNode Interface.
 	virtual void AllocateDefaultPins() override;
 	virtual FText GetNodeTitle(ENodeTitleType::Type TitleType) const override;
+	virtual FText GetTooltipText() const override;
+	virtual void ReconstructNode() override;
+	virtual bool IsCompatibleWithGraph(const UEdGraph* TargetGraph) const override;
 	virtual void GetNodeContextMenuActions(class UToolMenu* Menu, class UGraphNodeContextMenuContext* Context) const override;
 	virtual void PostPlacedNewNode() override;
 	virtual void PostPasteNode() override;
@@ -38,34 +57,127 @@ public:
 	//~ End UEdGraphNode Interface.
 
 	// UK2Node interface
+	virtual void ExpandNode(FKismetCompilerContext& CompilerContext, UEdGraph* SourceGraph) override;
+	virtual FName GetCornerIcon() const override;
 	virtual void GetMenuActions(FBlueprintActionDatabaseRegistrar& ActionRegistrar) const override;
 	virtual void PostReconstructNode() override;
 	// End of UK2Node interface
 
+	/** Set the operation call type. If it differs from the current, the node will be re-arranged to suit the new type. */
 	void SetAsyncType(EWebAPIOperationAsyncType InAsyncType);
 
+	/** Validity of this node. */
 	bool IsValid() const;
 
+protected:
+	// Returns the factory function (checked)
+	virtual UFunction* GetFactoryFunction() const;
+
+	/** Determines what the possible redirect pin names are **/
+	virtual void GetRedirectPinNames(const UEdGraphPin& Pin, TArray<FString>& RedirectPinNames) const override;
+
+	/**
+	 * If a the DefaultToSelf pin exists then it needs an actual connection to get properly casted
+	 * during compilation.
+	 *
+	 * @param InCompilerContext			The current compiler context used during expansion
+	 * @param InSourceGraph				The graph to place the expanded self node on
+	 * @param InIntermediateProxyNode		The spawned intermediate proxy node that has a DefaultToSelfPin
+	 *
+	 * @return	True if a successful connection was made to an intermediate node
+	 *			or if one was not necessary. False if a connection was failed.
+	 */
+	bool ExpandDefaultToSelfPin(FKismetCompilerContext& InCompilerContext, UEdGraph* InSourceGraph, UK2Node_CallFunction* InIntermediateProxyNode);
+
+	/** Expand out the logic to handle the delegate output pins */
+	virtual bool HandleDelegates(
+		const TArray<UE::WebAPI::Private::FOutputPinAndLocalVariable>& VariableOutputs, UEdGraphPin* ProxyObjectPin,
+		UEdGraphPin*& InOutLastThenPin, UEdGraph* SourceGraph, FKismetCompilerContext& CompilerContext);
+
+	/** Checks and caches outcome (success, error) delegates, returning true if both resolved. */
+	bool CacheOutcomeDelegates();
+
+	/** Invalidates current pin tool tips, so that they will be refreshed before being displayed: */
+	void InvalidatePinTooltips();
+
 private:
-	UEdGraphPin* FindPinContainingName(const EEdGraphPinDirection& InPinDirection, const FString& InName) const;
-	TArray<UEdGraphPin*> FindPinsContainingName(const EEdGraphPinDirection& InPinDirection, const FString& InName) const;
-	
+	friend struct FLatentActionPinMap;
+	friend struct FCallbackActionPinMap;
+	friend struct FOutcomeEventPinMap;
+
+	/** Find a single pin according to various criteria. */
+	UEdGraphPin* FindPin(
+		const FName& InName,
+		const EEdGraphPinDirection& InDirection = EEdGraphPinDirection::EGPD_MAX,
+		const FName& InCategory = NAME_All,
+		bool bFindPartial = false) const;
+
+	/** Find multiple pins according to various criteria. */
+	TArray<UEdGraphPin*> FindPins(
+		const FString& InName,
+		const EEdGraphPinDirection& InDirection = EEdGraphPinDirection::EGPD_MAX,
+		bool bOnlySplitPins = true) const;
+
+	/** Get all input pins that correspond to operation parameters. */
 	TArray<UEdGraphPin*> GetRequestPins() const;
+
+	/** Get both input request delegate pins, will return empty if this is Latent. */
+	TArray<UEdGraphPin*> GetRequestDelegatePins() const;
+
+	/** Get all (non-exec, non-error) output pins, will return empty if this uses Callbacks. */
 	TArray<UEdGraphPin*> GetResponsePins() const;
+
+	/** Get the output execution pin. */
+	UEdGraphPin* GetThenPin() const;
+
+	/** Get both outcome exec pins, will return empty if this uses Callbacks. */
+	TArray<UEdGraphPin*> GetResponseExecPins() const;
+
+	/** Get all error output pins, will return empty if this uses Callbacks. */
 	TArray<UEdGraphPin*> GetErrorResponsePins() const;
 
-	void SplitRequestPins(const TArray<UEdGraphPin*>& InPins) const;	
-	void SplitResponsePins(const TArray<UEdGraphPin*>& InPins) const;
+	/** Make and return a custom event in either the calling graph (if not function), or event graph. */
+	class UK2Node_CustomEvent* MakeCustomEvent(const FName& InName,
+		const UFunction* InSignature,
+		const FVector2D& InPosition = {});
 
-	static void CleanupPinNameInline(FString& InPinName);
+	/** Get or make a custom event node for the supplied outcome.  */
+	static const UK2Node_CustomEvent* GetCustomEventForOutcomeDelegate(const UK2Node_WebAPIOperation* InNode, const FName& InOutcomeName);
+
+	/** Called when converting FROM latent TO callback. */
+	void ConvertLatentToCallback();
+
+	/** Called when converting FROM callback TO latent. */
+	void ConvertCallbackToLatent();
 
 	/** Flips the node's Async Type. */
 	void ToggleAsyncType();
 
-	/** Update exec pins when converting between LatentAction and Callback Async Types. */
-	bool ReconnectExecPins(TArray<UEdGraphPin*>& OldPins);
+	/** The class containing the operation functions, AND is the return type of the latent function. */
+	UPROPERTY()
+	TObjectPtr<UClass> OperationClass;
+
+	/** The name of the latent function to call to create an operation object. */
+	UPROPERTY()
+	FName LatentFunctionName;
+	
+	/** The name of the delegated function to call to create an operation object. */
+	UPROPERTY()
+	FName DelegatedFunctionName;
+
+	/** Reference to the operations' positive outcome delegate property. */
+	TWeakFieldPtr<FMulticastDelegateProperty> PositiveDelegateProperty;
+
+	/** Reference to the operations' negative outcome delegate property. */
+	TWeakFieldPtr<FMulticastDelegateProperty> NegativeDelegateProperty;
 
 	/** Latent Action is preferred, but not compatible when used in a function, so allow conversion between the two. */
 	UPROPERTY()
 	EWebAPIOperationAsyncType OperationAsyncType = EWebAPIOperationAsyncType::LatentAction;
+
+	/** Constructing FText strings can be costly, so we cache the node's tooltip */
+	FNodeTextCache CachedTooltip;
+
+	/** Flag used to track validity of pin tooltips, when tooltips are invalid they will be refreshed before being displayed */
+	mutable bool bPinTooltipsValid;
 };
diff --git a/Engine/Plugins/Web/WebAPI/Source/WebAPIBlueprintGraph/Private/WebAPIBlueprintGraphModule.cpp b/Engine/Plugins/Web/WebAPI/Source/WebAPIBlueprintGraph/Private/WebAPIBlueprintGraphModule.cpp
--- a/Engine/Plugins/Web/WebAPI/Source/WebAPIBlueprintGraph/Private/WebAPIBlueprintGraphModule.cpp
+++ b/Engine/Plugins/Web/WebAPI/Source/WebAPIBlueprintGraph/Private/WebAPIBlueprintGraphModule.cpp
@@ -1,5 +1,24 @@
 // Copyright Epic Games, Inc. All Rights Reserved.
 
+#include "BlueprintActionFilter.h"
+#include "K2Node_WebAPIOperation.h"
 #include "Modules/ModuleManager.h"
 
-IMPLEMENT_MODULE(FDefaultModuleImpl, WebAPIBlueprintGraph)
+class FWebAPIBlueprintGraphModule final
+	: public IModuleInterface
+{
+public:
+	virtual void StartupModule() override
+	{
+		// Add actions that are relevant to the bound object from the pin class
+		FBlueprintActionFilter WebAPIOperationFilter(FBlueprintActionFilter::BPFILTER_NoFlags);
+		WebAPIOperationFilter.PermittedNodeTypes.Add(UK2Node_WebAPIOperation::StaticClass());
+		
+
+		FBlueprintActionFilter::AddUnique(WebAPIOperationFilter.TargetClasses, UK2Node_WebAPIOperation::StaticClass());
+	}
+	
+	virtual void ShutdownModule() override { }
+};
+
+IMPLEMENT_MODULE(FWebAPIBlueprintGraphModule, WebAPIBlueprintGraph)
diff --git a/Engine/Plugins/Web/WebAPI/Source/WebAPIBlueprintGraph/Private/WebAPIBlueprintGraphUtilities.cpp b/Engine/Plugins/Web/WebAPI/Source/WebAPIBlueprintGraph/Private/WebAPIBlueprintGraphUtilities.cpp
new file mode 100644
--- /dev/null
+++ b/Engine/Plugins/Web/WebAPI/Source/WebAPIBlueprintGraph/Private/WebAPIBlueprintGraphUtilities.cpp
@@ -0,0 +1,510 @@
+﻿// Copyright Epic Games, Inc. All Rights Reserved.
+
+#include "WebAPIBlueprintGraphUtilities.h"
+
+#include "EdGraphSchema_K2.h"
+#include "K2Node.h"
+#include "WebAPIOperationObject.h"
+#include "Algo/AllOf.h"
+#include "Algo/Transform.h"
+#include "UObject/WeakFieldPtr.h"
+
+namespace UE::WebAPI
+{
+	namespace Operation
+	{
+		template <uint32 Index, typename TEnableIf<Index < 2>::Type* = nullptr>
+		FMulticastDelegateProperty* GetOutcomeDelegate(const TSubclassOf<UWebAPIOperationObject>& InOperationClass, FName InOutcomeName)
+		{
+			check(InOperationClass);
+
+			FCachedOutcome* CachedOutcome = nullptr;
+
+			// If cached found
+			if(FCachedOutcome* Found = CachedOutcomeDelegates.Find(InOperationClass->GetFName()))
+			{
+				//  If valid, just return it
+				if(Found->Get<Index>().IsValid())
+				{
+					delete CachedOutcome;
+					CachedOutcome = nullptr;
+					
+					return Found->Get<Index>().Get();
+				}
+				else
+				{
+					// Set the pointer for later use
+					CachedOutcome = Found;
+				}
+			}
+			// Otherwise add it as a cache entry
+			else
+			{
+				CachedOutcome = &CachedOutcomeDelegates.Emplace(InOperationClass->GetFName());
+			}
+
+			int32 C = 0;
+			FMulticastDelegateProperty* P = nullptr;
+			const FString OutcomeNameStr = InOutcomeName.ToString();
+			for (TFieldIterator<FMulticastDelegateProperty> PropertyIterator(InOperationClass); PropertyIterator; ++PropertyIterator)
+			{
+				if(PropertyIterator->GetName().Contains(OutcomeNameStr))
+				{
+					++C;
+					P = *PropertyIterator;
+					CachedOutcome->Set<Index>(*PropertyIterator);
+					return P;
+				}
+			}
+
+			CachedOutcome = nullptr;
+			return nullptr;
+		}
+		
+		FMulticastDelegateProperty* GetPositiveOutcomeDelegate(const TSubclassOf<UWebAPIOperationObject>& InOperationClass)
+		{
+			check(InOperationClass);
+
+			return GetOutcomeDelegate<0>(InOperationClass, PositiveOutcomeName);
+		}
+
+		FMulticastDelegateProperty* GetNegativeOutcomeDelegate(const TSubclassOf<UWebAPIOperationObject>& InOperationClass)
+		{
+			check(InOperationClass);
+
+			return GetOutcomeDelegate<1>(InOperationClass, NegativeOutcomeName);
+		}
+
+		UFunction* GetOutcomeDelegateSignatureFunction(const TSubclassOf<UWebAPIOperationObject>& InOperationClass)
+		{
+			return GetPositiveOutcomeDelegate(InOperationClass)->SignatureFunction;			
+		}
+	}
+
+	namespace Graph
+	{
+		// @note: the next 4 calls are from FBlueprintEditor
+	
+		void CollectExecDownstreamNodes(const UEdGraphNode* CurrentNode, TArray<UEdGraphNode*>& CollectedNodes)
+		{
+			const UEdGraphSchema_K2* K2Schema = GetDefault<UEdGraphSchema_K2>();
+
+			TArray<UEdGraphPin*> AllPins = CurrentNode->GetAllPins();
+
+			for (UEdGraphPin*& Pin : AllPins)
+			{
+				if (Pin->Direction == EGPD_Output && Pin->PinType.PinCategory == K2Schema->PC_Exec)
+				{
+					for (UEdGraphPin*& Link : Pin->LinkedTo)
+					{
+						UEdGraphNode* LinkedNode = Cast<UEdGraphNode>(Link->GetOwningNode());
+						if (LinkedNode && !CollectedNodes.Contains(LinkedNode))
+						{
+							CollectedNodes.Add(LinkedNode);
+							CollectExecDownstreamNodes( LinkedNode, CollectedNodes );
+						}
+					}
+				}
+			}
+		}
+
+		void CollectExecUpstreamNodes(const UEdGraphNode* CurrentNode, TArray<UEdGraphNode*>& CollectedNodes)
+		{
+			const UEdGraphSchema_K2* K2Schema = GetDefault<UEdGraphSchema_K2>();
+
+			TArray<UEdGraphPin*> AllPins = CurrentNode->GetAllPins();
+
+			for (UEdGraphPin*& Pin : AllPins)
+			{
+				if (Pin->Direction == EGPD_Input && Pin->PinType.PinCategory == K2Schema->PC_Exec)
+				{
+					for (UEdGraphPin*& Link : Pin->LinkedTo)
+					{
+						UEdGraphNode* LinkedNode = Cast<UEdGraphNode>(Link->GetOwningNode());
+						if (LinkedNode && !CollectedNodes.Contains(LinkedNode))
+						{
+							CollectedNodes.Add(LinkedNode);
+							CollectExecUpstreamNodes( LinkedNode, CollectedNodes );
+						}
+					}
+				}
+			}
+		}
+
+		void CollectPureDownstreamNodes(const UEdGraphNode* CurrentNode, TArray<UEdGraphNode*>& CollectedNodes)
+		{
+			const UEdGraphSchema_K2* K2Schema = GetDefault<UEdGraphSchema_K2>();
+
+			TArray<UEdGraphPin*> AllPins = CurrentNode->GetAllPins();
+
+			for (UEdGraphPin*& Pin : AllPins)
+			{
+				if (Pin->Direction == EGPD_Output && Pin->PinType.PinCategory != K2Schema->PC_Exec)
+				{
+					for (UEdGraphPin*& Link : Pin->LinkedTo)
+					{
+						UK2Node* LinkedNode = Cast<UK2Node>(Link->GetOwningNode());
+						if (LinkedNode && !CollectedNodes.Contains(LinkedNode))
+						{
+							CollectedNodes.Add(LinkedNode);
+							if (LinkedNode->IsNodePure())
+							{
+								CollectPureDownstreamNodes( LinkedNode, CollectedNodes );
+							}
+						}
+					}
+				}
+			}
+		}
+
+		void CollectPureUpstreamNodes(const UEdGraphNode* CurrentNode, TArray<UEdGraphNode*>& CollectedNodes)
+		{
+			const UEdGraphSchema_K2* K2Schema = GetDefault<UEdGraphSchema_K2>();
+
+			TArray<UEdGraphPin*> AllPins = CurrentNode->GetAllPins();
+
+			for (UEdGraphPin*& Pin : AllPins)
+			{
+				if (Pin->Direction == EGPD_Input && Pin->PinType.PinCategory != K2Schema->PC_Exec)
+				{
+					for (UEdGraphPin*& Link : Pin->LinkedTo)
+					{
+						UK2Node* LinkedNode = Cast<UK2Node>(Link->GetOwningNode());
+						if (LinkedNode && !CollectedNodes.Contains(LinkedNode))
+						{
+							CollectedNodes.Add(LinkedNode);
+							if (LinkedNode->IsNodePure())
+							{
+								CollectPureUpstreamNodes( LinkedNode, CollectedNodes );
+							}
+						}
+					}
+				}
+			}
+		}
+		
+		void TransferPinConnections(
+			const UEdGraphPin* InSrc,
+			UEdGraphPin* InDst)
+		{
+			TArray<UEdGraphPin*> LinkedPins = InSrc->LinkedTo;
+			for(UEdGraphPin* LinkedPin : LinkedPins)
+			{
+				// Link Dst to LinkedPin
+				InDst->GetSchema()->TryCreateConnection(InDst, LinkedPin);
+			}
+		}
+
+		// For every src pin, get the corresponding dst pin - split if necessary
+		void TransferPins(
+			const TArray<UEdGraphPin*>& InSrcPins,
+			const TArray<UEdGraphPin*>& InDstPins)
+		{
+			for(UEdGraphPin* SrcPin : InSrcPins)
+			{
+				if(SrcPin->bHidden)
+				{
+					continue;
+				}
+		
+				// root = not split
+				const bool bIsRootPin = !SrcPin->ParentPin && SrcPin->SubPins.IsEmpty();
+				if(bIsRootPin)
+				{
+					if(UEdGraphPin* const* MatchingDstPin = InDstPins.FindByPredicate(
+						[SrcPinName = SrcPin->PinName](const UEdGraphPin* InPin)
+						{
+							return InPin->PinName == SrcPinName;
+						}))
+					{
+						UE::WebAPI::Graph::TransferPinConnections(SrcPin, *MatchingDstPin);
+					}
+				}
+				else
+				{
+					TArray Hierarchy = { SrcPin->PinName };
+
+					// Find root pin
+					const UEdGraphPin* RootPin = SrcPin->ParentPin;
+					while(RootPin != nullptr)
+					{
+						Hierarchy.Add(RootPin->PinName);
+						RootPin = RootPin->ParentPin;
+					}
+
+					Algo::Reverse(Hierarchy);
+
+					// Confined to current hierachy/subpins
+					TArray<UEdGraphPin*> DstPinsToSearch = InDstPins;
+					for(int32 PinIdx = 0; PinIdx < Hierarchy.Num(); ++PinIdx)
+					{
+						FName PinName = Hierarchy[PinIdx];
+				
+						if(UEdGraphPin* const* MatchingDstPin = DstPinsToSearch.FindByPredicate(
+							[SrcPinName = PinName](const UEdGraphPin* InPin)
+							{
+								return InPin->PinName == SrcPinName;
+							}))
+						{
+							const bool bIsLastItem = PinIdx == Hierarchy.Num() - 1;
+							if(bIsLastItem)
+							{
+								TransferPinConnections(SrcPin, *MatchingDstPin);
+								continue;
+							}
+					
+							if((*MatchingDstPin)->SubPins.IsEmpty())
+							{
+								(*MatchingDstPin)->GetSchema()->SplitPin(*MatchingDstPin);
+						
+								// Removes the prefix from the auto-named split pins 
+								if ((*MatchingDstPin)->SubPins.Num() > 0)
+								{
+									for (UEdGraphPin* SubPin : (*MatchingDstPin)->SubPins)
+									{
+										FString SubPinName = SubPin->PinName.ToString();
+										CleanupPinNameInline(SubPinName);
+
+										//SubPin->Modify();			
+										SubPin->PinFriendlyName = FText::FromString(SubPinName);
+									}
+								}
+							}
+
+							DstPinsToSearch = (*MatchingDstPin)->SubPins;
+						}
+					}
+				}
+			}
+		}
+
+		bool SplitPins(const TArray<UEdGraphPin*>& InPins)
+		{
+			// Compilation errors or stale nodes should be warnings, not errors
+			if(!ensureAlways(!InPins.IsEmpty() && Algo::AllOf(InPins, [](const UEdGraphPin* InPin) { return InPin; })))
+			{
+				return false;
+			}
+
+			bool bAllGood = true;
+			for(UEdGraphPin* Pin : InPins)
+			{
+				bAllGood &= SplitPin(Pin);	
+			}
+
+			return bAllGood;
+		}
+
+		bool SplitPin(UEdGraphPin* InPin)
+		{
+			bool bWasSplit = false;
+			if (const UScriptStruct* PinStructType = Cast<UScriptStruct>(InPin->PinType.PinSubCategoryObject.Get()))
+			{
+				uint32 PropertyNum = 0;
+			
+				// Split if not already
+				if (InPin->SubPins.Num() == 0 && InPin->GetOwningNode()->CanSplitPin(InPin))
+				{
+					for (TFieldIterator<const FProperty> PropertyIt(PinStructType, EFieldIteratorFlags::ExcludeSuper, EFieldIteratorFlags::ExcludeDeprecated); PropertyIt; ++PropertyIt)
+					{
+						if (PropertyIt->HasAnyPropertyFlags(CPF_Transient))
+						{
+							continue;
+						}
+
+						PropertyNum++;
+					}
+
+					if(PropertyNum == 1)
+					{
+						// Split, result should be single pin (cause it only contained a single property!)
+						const UEdGraphSchema_K2* K2Schema = GetDefault<UEdGraphSchema_K2>();
+						K2Schema->SplitPin(InPin);
+					}
+				}
+
+				// Removes the prefix from the auto-named split pins 
+				if (InPin->SubPins.Num() > 0)
+				{
+					bWasSplit = true;
+					for (UEdGraphPin* SubPin : InPin->SubPins)
+					{
+						FString SubPinName = SubPin->PinName.ToString();
+						CleanupPinNameInline(SubPinName);
+	
+						SubPin->PinFriendlyName = FText::FromString(SubPinName);
+					}
+				}
+			}
+
+			return bWasSplit;
+		}
+
+		TArray<UEdGraphPin*> FilterPinsByRelated(
+			const UEdGraphPin* InExecutionPin,
+			const TArray<UEdGraphPin*>& InPinsToFilter)
+		{
+			TMap<FName, UEdGraphPin*> NameToPin; // key = pin name
+			Algo::Transform(InPinsToFilter, NameToPin,
+				[](UEdGraphPin* InPin)
+				{
+					return TPair<FName, UEdGraphPin*>(InPin->PinName, InPin);
+				});
+			
+			TMap<FName, TArray<UEdGraphPin*>> LinkedPins; // key = pin name
+			Algo::Transform(InPinsToFilter, LinkedPins,
+				[](const UEdGraphPin* InPin)
+				{
+					return TPair<FName, TArray<UEdGraphPin*>>(InPin->PinName, InPin->LinkedTo);
+				});
+			
+			TMap<FName, TSet<UEdGraphNode*>> LinkedNodes; // key = pin name
+			Algo::Transform(InPinsToFilter, LinkedNodes,
+				[](const UEdGraphPin* InPin)
+				{
+					TSet<UEdGraphNode*> LinkedToNodes;
+					for(const UEdGraphPin* LinkedTo : InPin->LinkedTo)
+					{
+						LinkedToNodes.Add(LinkedTo->GetOwningNode());				
+					}
+					
+					return TPair<FName, TSet<UEdGraphNode*>>(InPin->PinName, LinkedToNodes);
+				});
+
+			TFunction<void(const UEdGraphPin*, TArray<UEdGraphNode*>&)> AppendConnectedToExec;
+			AppendConnectedToExec = [&AppendConnectedToExec](const UEdGraphPin* InPin, TArray<UEdGraphNode*>& OutNodes)
+			{
+				for(const UEdGraphPin* LinkedPin : InPin->LinkedTo)
+				{
+					UEdGraphNode* LinkedNode = LinkedPin->GetOwningNode();
+					OutNodes.Add(LinkedNode);
+
+					TArray<UEdGraphPin*> LinkedNodeExecPins = LinkedNode->GetAllPins().FilterByPredicate([](const UEdGraphPin* InLinkedNodePin)
+					{
+						return InLinkedNodePin->Direction == EGPD_Output
+							&& InLinkedNodePin->PinType.PinCategory == UEdGraphSchema_K2::PC_Exec;
+					});
+
+					for(const UEdGraphPin* LinkedNodeExecPin : LinkedNodeExecPins)
+					{
+						AppendConnectedToExec(LinkedNodeExecPin, OutNodes);
+					}
+				}
+			};
+
+			TArray<UEdGraphNode*> NodesInExecutionChain;
+			AppendConnectedToExec(InExecutionPin, NodesInExecutionChain);
+			
+			TArray<UEdGraphNode*> RelatedNodes;
+
+			TArray<UEdGraphNode*> ImpureNodes = NodesInExecutionChain.FilterByPredicate([](UEdGraphNode* InNode){
+				const UK2Node* K2Node = Cast<UK2Node>(InNode);
+				if (K2Node)
+				{
+					return !K2Node->IsNodePure();
+				}
+				return false;
+			});
+
+			TArray<UEdGraphNode*> PureNodes = NodesInExecutionChain.FilterByPredicate([](UEdGraphNode* InNode){
+				const UK2Node* K2Node = Cast<UK2Node>(InNode);
+				if (K2Node)
+				{
+					return K2Node->IsNodePure();
+				}
+				// Treat a node which can't cast to an UK2Node as a pure node (like a document node or a commment node)
+				// Make sure all selected nodes are handled
+				return true;
+			});
+
+			for (UEdGraphNode* ImpureNode : ImpureNodes)
+			{
+				RelatedNodes.Add(ImpureNode);
+				CollectExecDownstreamNodes(ImpureNode, RelatedNodes);
+				CollectExecUpstreamNodes(ImpureNode, RelatedNodes);
+				CollectPureDownstreamNodes(ImpureNode, RelatedNodes);
+				CollectPureUpstreamNodes(ImpureNode, RelatedNodes);
+			}
+
+			for (UEdGraphNode* PureNode : PureNodes)
+			{
+				RelatedNodes.Add(PureNode);
+				CollectPureDownstreamNodes(PureNode, RelatedNodes);
+				CollectPureUpstreamNodes(PureNode, RelatedNodes);
+			}
+
+			TSet<UEdGraphNode*> RelatedNodesSet{RelatedNodes};
+			TArray<UEdGraphPin*> UsedPins;
+
+			for(const TPair<FName, TSet<UEdGraphNode*>>& PinNodes : LinkedNodes)
+			{
+				// If any nodes for this pin are in the related nodes (for the current outcome execution)
+				if(!PinNodes.Value.Intersect(RelatedNodesSet).IsEmpty())
+				{
+					UsedPins.Add(NameToPin[PinNodes.Key]);			
+				}
+			}
+
+			return UsedPins;
+		}
+
+		UEdGraphPin* FindPin(
+			const UEdGraphNode* InNode,
+			const FName& InName,
+			const EEdGraphPinDirection& InDirection,
+			const FName& InCategory,
+			bool bFindPartial)
+		{
+			FString NameStr = InName.ToString();
+	
+			return InNode->FindPinByPredicate([&InName, &InDirection, &InCategory, bFindPartial, &NameStr](const UEdGraphPin* InPin)
+			{
+				return (bFindPartial ? InPin->PinName.ToString().Contains(NameStr) : InPin->PinName == InName)
+					&& (InDirection < EEdGraphPinDirection::EGPD_MAX ? InPin->Direction == InDirection : true)
+					&& (InCategory != NAME_All ? InPin->PinType.PinCategory == InCategory : true);
+			});
+		}
+
+		TArray<UEdGraphPin*> FindPins(
+			const UEdGraphNode* InNode,
+			const FString& InName,
+			const EEdGraphPinDirection& InPinDirection,
+			bool bOnlySplitPins)
+		{
+			return InNode->Pins.FilterByPredicate([&InPinDirection, InName, bOnlySplitPins](const UEdGraphPin* InPin)
+			{
+				return InPin->Direction == InPinDirection
+					&& (bOnlySplitPins ? InPin->ParentPin || InPin->SubPins.IsEmpty() : InPin->SubPins.IsEmpty())
+					&& InPin->PinType.PinCategory != UEdGraphSchema_K2::PC_Exec
+					&& InPin->GetName().Contains(InName);
+			});
+		}
+
+		TArray<UEdGraphPin*> GetResponsePins(const UEdGraphNode* InNode)
+		{
+			TArray<UEdGraphPin*> ResponsePins = FindPins(InNode, TEXT("Response"), EGPD_Output);
+			ensureMsgf(!ResponsePins.IsEmpty(), TEXT("The Operation must contain a delegate with a parameter containing \"Response\" in the name."));	
+			return ResponsePins;
+		}
+
+		TArray<UEdGraphPin*> GetErrorResponsePins(const UEdGraphNode* InNode)
+		{
+			static FString NegativeOutcomeNameStr = UE::WebAPI::Operation::NegativeOutcomeName.ToString();
+
+			TArray<UEdGraphPin*> ErrorResponsePins = FindPins(InNode, NegativeOutcomeNameStr, EGPD_Output);
+			ensureMsgf(!ErrorResponsePins.IsEmpty(), TEXT("The Operation must contain a delegate with a parameter containing \"%s\" in the name."), *NegativeOutcomeNameStr);	
+			return ErrorResponsePins;
+		}
+
+		void CleanupPinNameInline(FString& InPinName)
+		{
+			constexpr static TCHAR DelimiterChar = TCHAR('_');
+
+			int32 UnderscorePosition = -1;
+			if(InPinName.FindLastChar(DelimiterChar, UnderscorePosition))
+			{
+				InPinName = InPinName.RightChop(UnderscorePosition + 1);
+			}
+		}
+	}
+}
diff --git a/Engine/Plugins/Web/WebAPI/Source/WebAPIBlueprintGraph/Private/WebAPIBlueprintGraphUtilities.h b/Engine/Plugins/Web/WebAPI/Source/WebAPIBlueprintGraph/Private/WebAPIBlueprintGraphUtilities.h
new file mode 100644
--- /dev/null
+++ b/Engine/Plugins/Web/WebAPI/Source/WebAPIBlueprintGraph/Private/WebAPIBlueprintGraphUtilities.h
@@ -0,0 +1,122 @@
+﻿// Copyright Epic Games, Inc. All Rights Reserved.
+
+#pragma once
+
+#include "CoreMinimal.h"
+#include "EdGraph/EdGraphNode.h"
+#include "EdGraph/EdGraphPin.h"
+#include "Templates/SubclassOf.h"
+#include "UObject/Class.h"
+#include "UObject/UnrealType.h"
+#include "UObject/WeakFieldPtr.h"
+
+class UWebAPIOperationObject;
+
+namespace UE::WebAPI
+{
+	namespace Operation
+	{
+		/** Stores references to outcome delegate properties for a given operation. */
+		struct FCachedOutcome
+		{
+		public:
+			TWeakFieldPtr<FMulticastDelegateProperty> PositiveOutcomeDelegate;
+			TWeakFieldPtr<FMulticastDelegateProperty> NegativeOutcomeDelegate;
+
+			template <uint32 Index, typename TEnableIf<Index < 2>::Type* = nullptr>
+			const TWeakFieldPtr<FMulticastDelegateProperty>& Get() const
+			{
+				if constexpr (Index == 0)
+				{
+					return PositiveOutcomeDelegate;
+				}
+				else
+				{
+					return NegativeOutcomeDelegate;
+				}
+			}
+
+			template <uint32 Index, typename TEnableIf<Index < 2>::Type* = nullptr>
+			void Set(const FMulticastDelegateProperty* InDelegateProperty)
+			{
+				if constexpr (Index == 0)
+				{
+					PositiveOutcomeDelegate = InDelegateProperty;
+				}
+				else
+				{
+					NegativeOutcomeDelegate = InDelegateProperty;
+				}
+			}
+		};
+
+		/** String denoting a positive outcome, should match generated delegate name. */
+		inline const static FName PositiveOutcomeName = TEXT("Success");
+
+		/** String denoting a negative outcome, should match generated delegate name. */
+		inline const static FName NegativeOutcomeName = TEXT("Error");
+
+		/** Cached positive and negative (in order) delegates per operation class (name). */
+		static TMap<FName, FCachedOutcome> CachedOutcomeDelegates;
+
+		FMulticastDelegateProperty* GetPositiveOutcomeDelegate(const TSubclassOf<UWebAPIOperationObject>& InOperationClass);
+
+		template <class OperationType = UWebAPIOperationObject>
+		FMulticastDelegateProperty* GetPositiveOutcomeDelegate()
+		{
+			return GetPositiveOutcomeDelegate(OperationType::StaticClass());			
+		}
+		
+		FMulticastDelegateProperty* GetNegativeOutcomeDelegate(const TSubclassOf<UWebAPIOperationObject>& InOperationClass);
+
+		template <class OperationType = UWebAPIOperationObject>
+		FMulticastDelegateProperty* GetNegativeOutcomeDelegate()
+		{
+			return GetNegativeOutcomeDelegate(OperationType::StaticClass());			
+		}
+
+		UFunction* GetOutcomeDelegateSignatureFunction(const TSubclassOf<UWebAPIOperationObject>& InOperationClass);
+	} 
+
+	namespace Graph
+	{
+		/** Connects all pins linked to src, to dst. */
+		void TransferPinConnections(const UEdGraphPin* InSrc, UEdGraphPin* InDst);
+
+		/** Connects all pins linked to all src pins, to matching (same named) dst pins. */
+		void TransferPins(const TArray<UEdGraphPin*>& InSrcPins, const TArray<UEdGraphPin*>& InDstPins);
+
+		/** Split multiple pins according to various criteria, like unwrapping single property structs, etc. */
+		bool SplitPins(const TArray<UEdGraphPin*>& InPins);
+
+		/** Split single pin according to various criteria, like unwrapping single property structs, etc. */
+		bool SplitPin(UEdGraphPin* InPin);
+
+		/** Filter the inputs pins by the input execution chain (like "Hide Unrelated"). */
+		TArray<UEdGraphPin*> FilterPinsByRelated(const UEdGraphPin* InExecutionPin, const TArray<UEdGraphPin*>& InPinsToFilter);
+
+		/** Find a single pin according to various criteria. */
+		UEdGraphPin* FindPin(
+			const UEdGraphNode* InNode,
+			const FName& InName,
+			const EEdGraphPinDirection& InDirection = EEdGraphPinDirection::EGPD_MAX,
+			const FName& InCategory = NAME_All,
+			bool bFindPartial = false);
+
+		/** Find multiple pins according to various criteria. */
+		TArray<UEdGraphPin*> FindPins(
+			const UEdGraphNode* InNode,
+			const FString& InName,
+			const EEdGraphPinDirection& InPinDirection,
+			bool bOnlySplitPins = true);
+
+		/** Get all (non-exec, non-error) output pins, will return empty if this uses Callbacks. */
+		TArray<UEdGraphPin*> GetResponsePins(const UEdGraphNode* InNode);
+
+		/** Get all error output pins, will return empty if this uses Callbacks. */
+		TArray<UEdGraphPin*> GetErrorResponsePins(const UEdGraphNode* InNode);
+
+		/** Returns the last part of a string, split by underscores. */
+		void CleanupPinNameInline(FString& InPinName);
+	}
+}
diff --git a/Engine/Plugins/Web/WebAPI/Source/WebAPIEditor/Private/Assets/WebAPIDefinitionAssetEditorToolkit.cpp b/Engine/Plugins/Web/WebAPI/Source/WebAPIEditor/Private/Assets/WebAPIDefinitionAssetEditorToolkit.cpp
--- a/Engine/Plugins/Web/WebAPI/Source/WebAPIEditor/Private/Assets/WebAPIDefinitionAssetEditorToolkit.cpp
+++ b/Engine/Plugins/Web/WebAPI/Source/WebAPIEditor/Private/Assets/WebAPIDefinitionAssetEditorToolkit.cpp
@@ -1,4 +1,4 @@
-// Copyright Epic Games, Inc. All Rights Reserved.
+﻿// Copyright Epic Games, Inc. All Rights Reserved.
 
 #include "WebAPIDefinitionAssetEditorToolkit.h"
 
@@ -17,9 +17,9 @@
 #include "Framework/Docking/TabManager.h"
 #include "Framework/MultiBox/MultiBoxBuilder.h"
 #include "Modules/ModuleManager.h"
+#include "Styling/AppStyle.h"
 #include "Toolkits/AssetEditorToolkit.h"
 #include "Widgets/Docking/SDockTab.h"
-#include "Styling/AppStyle.h"
 
 #define LOCTEXT_NAMESPACE "WebAPIDefinitionAssetEditorToolkit"
 
@@ -35,6 +35,7 @@ const FName FWebAPIDefinitionAssetEditorToolkit::CodeTabID(TEXT("WebAPIDefinitio
 FWebAPIDefinitionAssetEditorToolkit::FWebAPIDefinitionAssetEditorToolkit()
 	: Definition(nullptr)
 {
+	LayoutExtender = MakeShared<FLayoutExtender>();
 }
 
 FWebAPIDefinitionAssetEditorToolkit::~FWebAPIDefinitionAssetEditorToolkit()
@@ -89,6 +90,9 @@ void FWebAPIDefinitionAssetEditorToolkit::Initialize(const EToolkitMode::Type In
 	
 	RegenerateMenusAndToolbars();
 
+	WebAPIEditorModule.OnRegisterLayoutExtensions().Broadcast(*LayoutExtender);
+	StandaloneDefaultLayout->ProcessExtensions(*LayoutExtender);
+
 	// Support undo/redo
 	GEditor->RegisterForUndo(this);
 
diff --git a/Engine/Plugins/Web/WebAPI/Source/WebAPIEditor/Private/Assets/WebAPIDefinitionAssetEditorToolkit.h b/Engine/Plugins/Web/WebAPI/Source/WebAPIEditor/Private/Assets/WebAPIDefinitionAssetEditorToolkit.h
--- a/Engine/Plugins/Web/WebAPI/Source/WebAPIEditor/Private/Assets/WebAPIDefinitionAssetEditorToolkit.h
+++ b/Engine/Plugins/Web/WebAPI/Source/WebAPIEditor/Private/Assets/WebAPIDefinitionAssetEditorToolkit.h
@@ -90,8 +90,7 @@ protected:
 	
 	/** Code View. */
 	TSharedPtr<class SWebAPICodeView> CodeView;
-	
-	// @todo: implement so users can save layouts
+
 	/** Extender for adding to the default layout for this asset editor */
 	TSharedPtr<FLayoutExtender> LayoutExtender;
 
diff --git a/Engine/Plugins/Web/WebAPI/Source/WebAPIEditor/Private/Details/ViewModels/WebAPIOperationViewModel.cpp b/Engine/Plugins/Web/WebAPI/Source/WebAPIEditor/Private/Details/ViewModels/WebAPIOperationViewModel.cpp
--- a/Engine/Plugins/Web/WebAPI/Source/WebAPIEditor/Private/Details/ViewModels/WebAPIOperationViewModel.cpp
+++ b/Engine/Plugins/Web/WebAPI/Source/WebAPIEditor/Private/Details/ViewModels/WebAPIOperationViewModel.cpp
@@ -112,24 +112,6 @@ void FWebAPIOperationViewModel::Initialize()
 		return Lhs->GetLabel().ToString() < Rhs->GetLabel().ToString();
 	});
 
-	// @todo:
-	/*
-	const TArray<TObjectPtr<UWebAPIOperationParameter>> Parameters = Operation->Parameters;
-	TArray<FString> ParameterLabels;
-	Algo::Transform(Parameters, ParameterLabels, [&](const TObjectPtr<UWebAPIOperationParameter>& InParameter)
-	{
-		return FString::Printf(TEXT("%s: %s"), *InParameter->Name.ToString(), *InParameter->Type.ToString());
-	});
-	*/
-
-	/*
-	const FWebAPINameVariant Name = Operation->Name;
-	CachedLabel = FText::Format(LOCTEXT("OperationLabel", "{0}({1}): {2}"),
-		FText::FromString(Name.ToString()),
-		FText::FromString(FString::Join(ParameterLabels, TEXT(", "))),
-		FText::FromString(TEXT("void")));
-		*/
-
 	const FString Name = Operation->Name.GetDisplayName();
 	CachedLabel = FText::Format(LOCTEXT("OperationLabel", "{0}"),
 		FText::FromString(Name));
diff --git a/Engine/Plugins/Web/WebAPI/Source/WebAPIEditor/Private/Details/WebAPIDefinitionTargetModuleCustomization.cpp b/Engine/Plugins/Web/WebAPI/Source/WebAPIEditor/Private/Details/WebAPIDefinitionTargetModuleCustomization.cpp
--- a/Engine/Plugins/Web/WebAPI/Source/WebAPIEditor/Private/Details/WebAPIDefinitionTargetModuleCustomization.cpp
+++ b/Engine/Plugins/Web/WebAPI/Source/WebAPIEditor/Private/Details/WebAPIDefinitionTargetModuleCustomization.cpp
@@ -78,7 +78,6 @@ void FWebAPIDefinitionTargetModuleCustomization::OnAddModule()
 
 bool FWebAPIDefinitionTargetModuleCustomization::CanAddModule() const
 {
-	// @todo: will this ever be false? if not, remove
 	return true;
 }
 
diff --git a/Engine/Plugins/Web/WebAPI/Source/WebAPIEditor/Private/Details/Widgets/SWebAPISchemaOperationRow.cpp b/Engine/Plugins/Web/WebAPI/Source/WebAPIEditor/Private/Details/Widgets/SWebAPISchemaOperationRow.cpp
--- a/Engine/Plugins/Web/WebAPI/Source/WebAPIEditor/Private/Details/Widgets/SWebAPISchemaOperationRow.cpp
+++ b/Engine/Plugins/Web/WebAPI/Source/WebAPIEditor/Private/Details/Widgets/SWebAPISchemaOperationRow.cpp
@@ -21,7 +21,7 @@ void SWebAPISchemaOperationRow::Construct(const FArguments& InArgs, const TShare
 		[
 			SNew(SBorder)
 			.BorderBackgroundColor(FLinearColor(0,0,0,0))
-			// @todo: this prevents selection highlight
+			// @todo: this prevents selection highlight, check slack correspondence
 			//.BorderImage(FAppStyle::GetBrush("ToolPanel.GroupBorder"))
 			.Padding(4)
 			[
diff --git a/Engine/Plugins/Web/WebAPI/Source/WebAPIEditor/Private/Details/Widgets/SWebAPISchemaServiceRow.cpp b/Engine/Plugins/Web/WebAPI/Source/WebAPIEditor/Private/Details/Widgets/SWebAPISchemaServiceRow.cpp
--- a/Engine/Plugins/Web/WebAPI/Source/WebAPIEditor/Private/Details/Widgets/SWebAPISchemaServiceRow.cpp
+++ b/Engine/Plugins/Web/WebAPI/Source/WebAPIEditor/Private/Details/Widgets/SWebAPISchemaServiceRow.cpp
@@ -2,14 +2,13 @@
 
 #include "SWebAPISchemaServiceRow.h"
 
+#include "Styling/AppStyle.h"
 #include "WebAPIEditorStyle.h"
 #include "Widgets/Colors/SColorBlock.h"
 #include "Widgets/Input/SCheckBox.h"
 #include "Widgets/Input/SHyperlink.h"
 #include "Widgets/Layout/SBox.h"
 
-#include "Styling/AppStyle.h"
-
 void SWebAPISchemaServiceRow::Construct(const FArguments& InArgs, const TSharedRef<FWebAPIServiceViewModel>& InViewModel, const TSharedRef<STableViewBase>& InOwnerTableView)
 {
 	SWebAPISchemaTreeTableRow::Construct(
@@ -18,7 +17,7 @@ void SWebAPISchemaServiceRow::Construct(const FArguments& InArgs, const TSharedR
 		[
 			SNew(SBorder)
 			.BorderBackgroundColor(FLinearColor(0,0,0,0))
-			// @todo: this prevents selection highlight
+			// @todo: this prevents selection highlight, check slack correspondence
 			//.BorderImage(FAppStyle::GetBrush("ToolPanel.GroupBorder"))
 			.Padding(4)
 			[
diff --git a/Engine/Plugins/Web/WebAPI/Source/WebAPIEditor/Private/Dom/WebAPIOperation.cpp b/Engine/Plugins/Web/WebAPI/Source/WebAPIEditor/Private/Dom/WebAPIOperation.cpp
--- a/Engine/Plugins/Web/WebAPI/Source/WebAPIEditor/Private/Dom/WebAPIOperation.cpp
+++ b/Engine/Plugins/Web/WebAPI/Source/WebAPIEditor/Private/Dom/WebAPIOperation.cpp
@@ -69,18 +69,6 @@ void UWebAPIOperationRequest::Visit(TFunctionRef<void(IWebAPISchemaObjectInterfa
 	}
 }
 
-/*
-void UWebAPIOperationParameters::Visit(TFunctionRef<void(IWebAPISchemaObjectInterface*&)> InVisitor)
-{
-	IWebAPISchemaObjectInterface::Visit(InVisitor);
-
-	for(const TObjectPtr<UWebAPIOperationParameter>& Parameter : Parameters)
-	{
-		Parameter->Visit(InVisitor);		
-	}
-}
-*/
-
 UWebAPIOperation::UWebAPIOperation()
 {
 	Request = CreateDefaultSubobject<UWebAPIOperationRequest>(TEXT("Request"));
@@ -101,8 +89,6 @@ void UWebAPIOperation::Visit(TFunctionRef<void(IWebAPISchemaObjectInterface*&)>
 	IWebAPISchemaObjectInterface::Visit(InVisitor);
 	
 	Request->Visit(InVisitor);
-	TArray<int32> A;
-	int32& B = A.Emplace_GetRef();
 	for(const TObjectPtr<UWebAPIOperationResponse>& Response : Responses)
 	{
 		Response->Visit(InVisitor);		
@@ -114,7 +100,7 @@ void UWebAPIOperation::Visit(TFunctionRef<void(IWebAPISchemaObjectInterface*&)>
 	}
 }
 
-void UWebAPIOperation::BindToTypeInfo() const
+void UWebAPIOperation::BindToTypeInfo()
 {
 	check(Name.HasTypeInfo());
 
diff --git a/Engine/Plugins/Web/WebAPI/Source/WebAPIEditor/Private/Dom/WebAPITypeRegistry.cpp b/Engine/Plugins/Web/WebAPI/Source/WebAPIEditor/Private/Dom/WebAPITypeRegistry.cpp
--- a/Engine/Plugins/Web/WebAPI/Source/WebAPIEditor/Private/Dom/WebAPITypeRegistry.cpp
+++ b/Engine/Plugins/Web/WebAPI/Source/WebAPIEditor/Private/Dom/WebAPITypeRegistry.cpp
@@ -3,6 +3,7 @@
 #include "Dom/WebAPITypeRegistry.h"
 
 #include "GraphEditorSettings.h"
+#include "WebAPIEditorLog.h"
 #include "Algo/Copy.h"
 #include "Serialization/JsonTypes.h"
 
@@ -178,10 +179,8 @@ TObjectPtr<UWebAPITypeInfo> UWebAPITypeRegistry::GetOrMakeGeneratedType(
 	TypeInfo->DeclarationType = InDeclarationType;
 	TypeInfo->SetName(InName);
 
-	// @todo: remove
-	UE_LOG(LogTemp, Display, TEXT("New Type Path: %s"), *TypeInfo->GetPathName());
-
 #if WITH_WEBAPI_DEBUG
+	UE_LOG(LogWebAPIEditor, Display, TEXT("New Type Path: %s"), *TypeInfo->GetPathName());
 	TypeInfo->DebugString += TEXT(">Created in GetOrMakeGeneratedType");
 #endif
 
@@ -206,8 +205,9 @@ TObjectPtr<UWebAPITypeInfo> UWebAPITypeRegistry::GetOrMakeGeneratedType(
 {
 	ensureAlways(!InName.IsEmpty());
 
-	// @todo: remove
-	UE_LOG(LogTemp, Display, TEXT("Template Type Path: %s"), *InTemplateTypeInfo->GetPathName());	
+#if WITH_WEBAPI_DEBUG
+	UE_LOG(LogWebAPIEditor, Display, TEXT("Template Type Path: %s"), *InTemplateTypeInfo->GetPathName());	
+#endif
 	
 	if(const TObjectPtr<UWebAPITypeInfo>* FoundType = FindGeneratedType(InSchemaType, InName))
 	{
@@ -217,9 +217,10 @@ TObjectPtr<UWebAPITypeInfo> UWebAPITypeRegistry::GetOrMakeGeneratedType(
 	const FName TypeName = UE::WebAPI::MakeTypeInfoName(this, InSchemaType, InName);
 	GeneratedTypes.Emplace(DuplicateObject(InTemplateTypeInfo.Get(), this, TypeName));
 	const TObjectPtr<UWebAPITypeInfo>& TypeInfo = GeneratedTypes.Last();
-	
-	// @todo: remove
-	UE_LOG(LogTemp, Display, TEXT("Duplicate Type Path: %s"), *TypeInfo->GetPathName());
+
+#if WITH_WEBAPI_DEBUG
+	UE_LOG(LogWebAPIEditor, Display, TEXT("Duplicate Type Path: %s"), *TypeInfo->GetPathName());
+#endif
 
 	TypeInfo->Name = InName;
 	TypeInfo->SchemaType = InSchemaType;
@@ -283,11 +284,6 @@ bool UWebAPITypeRegistry::CheckAllNamed() const
 	TArray<TObjectPtr<UWebAPITypeInfo>> UnnamedTypes;
 	if(!CheckAllNamed(UnnamedTypes))
 	{
-		// @todo: what's this garbage?
-	// 	const FText LogMessage = FText::FormatNamed(
-	// LOCTEXT("HasUnnamed", "{Num} unnamed types found in the TypeRegistry."),
-	// 	TEXT("Num"), UnnamedTypes.Num());
-	// 	FWebAPIMessageLog::Get().LogError(LogMessage, UWebAPITypeRegistry::LogName);
 		return false;
 	}
 
@@ -333,8 +329,9 @@ FName UE::WebAPI::MakeTypeInfoName(UObject* InOuter, const EWebAPISchemaType& In
 		Result = MakeUniqueObjectName(InOuter, UWebAPITypeInfo::StaticClass(), FName(FString::Format(TEXT("{ClassName}_{SchemaType}_{TypeName}_"), NameArgs)));
 	}
 
-	// @todo: remove
-	UE_LOG(LogTemp, Display, TEXT("MakeTypeName: %s"), *Result.ToString());
+#if WITH_WEBAPI_DEBUG
+	UE_LOG(LogWebAPIEditor, Display, TEXT("MakeTypeName: %s"), *Result.ToString());
+#endif
 	
 	return Result;
 }
diff --git a/Engine/Plugins/Web/WebAPI/Source/WebAPIEditor/Private/Tests/WebAPIJsonTestData.h b/Engine/Plugins/Web/WebAPI/Source/WebAPIEditor/Private/Tests/WebAPIJsonTestData.h
new file mode 100644
--- /dev/null
+++ b/Engine/Plugins/Web/WebAPI/Source/WebAPIEditor/Private/Tests/WebAPIJsonTestData.h
@@ -0,0 +1,72 @@
+﻿// Copyright Epic Games, Inc. All Rights Reserved.
+
+#pragma once
+
+#include "CoreMinimal.h"
+#include "WebAPIJsonUtilities.h"
+
+#include "WebAPIJsonTestData.generated.h"
+
+UENUM()
+enum class EWAPITestUEnum : uint8
+{
+	EnumValueA = 4,
+	EnumValueB = 2,
+	EnumValueC = 7
+};
+
+namespace UE::WebAPI::Testing
+{
+	enum class EWAPITestEnum : uint8
+	{
+		EnumValueA = 4,
+		EnumValueB = 2,
+		EnumValueC = 7
+	};
+	
+	static TMap<FString, EWAPITestEnum> StringToTestEnum = {
+		{ TEXT("EnumValueA"), EWAPITestEnum::EnumValueA },
+		{ TEXT("EnumValueB"), EWAPITestEnum::EnumValueB },
+		{ TEXT("EnumValueC"), EWAPITestEnum::EnumValueC },
+	};
+
+	struct WEBAPIEDITOR_API FTestStruct
+	{
+		float TestFloat = 45.2f;
+		FText TestText = FText::FromString(TEXT("Test Text"));
+		TSharedPtr<FTestStruct> TestPtrOfContainingType;
+		TArray<TSharedPtr<FTestStruct>> ArrayOfOuterType;
+
+		bool FromJson(const TSharedPtr<FJsonObject>& InJsonObject)
+		{
+			// at least one must deserialize
+			bool bResult = false;
+			bResult |= Json::TryGetField(InJsonObject, GET_MEMBER_NAME_STRING_CHECKED(FTestStruct, TestFloat), TestFloat);
+			bResult |= Json::TryGetField(InJsonObject, GET_MEMBER_NAME_STRING_CHECKED(FTestStruct, TestText), TestText);
+			bResult |= Json::TryGetField(InJsonObject, GET_MEMBER_NAME_STRING_CHECKED(FTestStruct, TestPtrOfContainingType), TestPtrOfContainingType);
+			bResult |= Json::TryGetField(InJsonObject, GET_MEMBER_NAME_STRING_CHECKED(FTestStruct, ArrayOfOuterType), ArrayOfOuterType);
+			return bResult;
+		}
+	};
+	
+	class WEBAPIEDITOR_API FTestClass
+	{
+	public:
+		TArray<FString> ArrayOfStrings;
+		TArray<FTestStruct> ArrayOfStructs;
+		TArray<TSharedPtr<FTestStruct>> ArrayOfSharedPtr;
+		TArray<TOptional<FTestStruct>> ArrayOfOptional;
+		TArray<TSharedPtr<TOptional<FTestStruct>>> ArrayOfSharedPtrOptional;
+
+		EWAPITestUEnum UEnumValue = EWAPITestUEnum::EnumValueA;
+		EWAPITestEnum EnumValue = EWAPITestEnum::EnumValueA;
+		TSharedPtr<EWAPITestUEnum> EnumSharedPtr;
+
+		Json::TJsonReference<FTestStruct> JsonReferenceStruct;
+		Json::TJsonReference<TSharedPtr<FTestStruct>> JsonReferenceStructSharedPtr;
+
+		TMap<EWAPITestUEnum, TArray<float>> MapOfEnumToFloats;
+
+		TVariant<EWAPITestUEnum, FTestStruct, bool> Variant;
+	};
+}
diff --git a/Engine/Plugins/Web/WebAPI/Source/WebAPIEditor/Private/Tests/WebAPIJsonUtilities.spec.cpp b/Engine/Plugins/Web/WebAPI/Source/WebAPIEditor/Private/Tests/WebAPIJsonUtilities.spec.cpp
new file mode 100644
--- /dev/null
+++ b/Engine/Plugins/Web/WebAPI/Source/WebAPIEditor/Private/Tests/WebAPIJsonUtilities.spec.cpp
@@ -0,0 +1,356 @@
+﻿// Copyright Epic Games, Inc. All Rights Reserved.
+
+#include "Misc/AutomationTest.h"
+
+#if WITH_AUTOMATION_TESTS
+
+#include "Serialization/JsonReader.h"
+#include "Serialization/JsonSerializer.h"
+#include "WebAPIJsonTestData.h"
+#include "WebAPIJsonUtilities.h"
+
+BEGIN_DEFINE_SPEC(FWebAPIEditorJsonSpec,
+	TEXT("Plugin.WebAPI.Editor.Json"),
+	EAutomationTestFlags::EditorContext | EAutomationTestFlags::EngineFilter | EAutomationTestFlags::ApplicationContextMask)
+
+	TSharedPtr<FJsonObject> JsonObject;
+
+	// Deserialize to JsonObject
+	bool DeserializeJsonString(const FString& InJsonString)
+	{
+		const TSharedRef<TJsonReader<>> JsonReader = TJsonReaderFactory<>::Create(InJsonString);
+		return FJsonSerializer::Deserialize(JsonReader, JsonObject);
+	}
+
+END_DEFINE_SPEC(FWebAPIEditorJsonSpec)
+
+void FWebAPIEditorJsonSpec::Define()
+{
+	using namespace UE::WebAPI;
+		
+	Describe(TEXT("Deserialization"), [this]
+	{
+		Describe(TEXT("Enum"), [this]
+		{
+			It(TEXT("loads from string"), [this]
+			{
+				static constexpr const TCHAR* ExampleJson = TEXT(R"_JSON(
+				{
+				    "enumValue": "enumValueB"
+				})_JSON");
+				const bool bDeserialized = DeserializeJsonString(ExampleJson);
+				TestTrue(TEXT("Deserialized Json"), bDeserialized);
+
+				if(bDeserialized)
+				{
+					Testing::FTestClass TestObject;
+					const bool bGotField = UE::Json::TryGetField(
+						JsonObject,
+						GET_MEMBER_NAME_STRING_CHECKED(Testing::FTestClass, EnumValue),
+						TestObject.EnumValue,
+						Testing::StringToTestEnum);
+				
+					TestTrue(TEXT("TryGetField"), bGotField);
+					TestEqual(TEXT("EnumValue is EnumValueB"), TestObject.EnumValue, Testing::EWAPITestEnum::EnumValueB);
+				}
+			});
+
+			// UEnum doesn't require a String to Value lookup - it gets that from StaticEnum<>
+			It(TEXT("loads UEnum from string"), [this]
+			{
+				static constexpr const TCHAR* ExampleJson = TEXT(R"_JSON(
+				{
+				    "uenumValue": "enumValueC"
+				})_JSON");
+				const bool bDeserialized = DeserializeJsonString(ExampleJson);
+				TestTrue(TEXT("Deserialized Json"), bDeserialized);
+
+				if(bDeserialized)
+				{
+					Testing::FTestClass TestObject;
+					const bool bGotField = UE::Json::TryGetField(
+						JsonObject,
+						GET_MEMBER_NAME_STRING_CHECKED(Testing::FTestClass, UEnumValue),
+						TestObject.UEnumValue);
+				
+					TestTrue(TEXT("TryGetField"), bGotField);
+					TestEqual(TEXT("EnumValue is EnumValueC"), TestObject.UEnumValue, EWAPITestUEnum::EnumValueC);
+				}
+			});
+
+			It(TEXT("doesn't load from string that doesn't name a valid value"), [this]
+			{
+				static constexpr const TCHAR* ExampleJson = TEXT(R"_JSON(
+				{
+				    "enumValue": "invalidEnumValue"
+				})_JSON");
+				const bool bDeserialized = DeserializeJsonString(ExampleJson);
+				TestTrue(TEXT("Deserialized Json"), bDeserialized);
+	
+				if(bDeserialized)
+				{
+					Testing::FTestClass TestObject;
+					const bool bGotField = UE::Json::TryGetField(
+						JsonObject,
+						GET_MEMBER_NAME_STRING_CHECKED(Testing::FTestClass, EnumValue),
+						TestObject.EnumValue,
+						Testing::StringToTestEnum);
+								
+					TestFalse(TEXT("TryGetField"), bGotField);
+				}
+			});
+
+			It(TEXT("loads from integer"), [this]
+			{
+				static constexpr const TCHAR* ExampleJson = TEXT(R"_JSON(
+				{
+				    "enumValue": 7
+				})_JSON");
+				const bool bDeserialized = DeserializeJsonString(ExampleJson);
+				TestTrue(TEXT("Deserialized Json"), bDeserialized);
+
+				if(bDeserialized)
+				{
+					Testing::FTestClass TestObject;
+					const bool bGotField = UE::Json::TryGetField(
+						JsonObject,
+						GET_MEMBER_NAME_STRING_CHECKED(Testing::FTestClass, EnumValue),
+						TestObject.EnumValue,
+						Testing::StringToTestEnum);
+												
+					TestTrue(TEXT("TryGetField"), bGotField);
+					TestEqual(TEXT("EnumValue is EnumValueC"), TestObject.EnumValue, Testing::EWAPITestEnum::EnumValueC);
+				}
+			});
+		});
+
+		Describe(TEXT("Arrays"), [this]
+		{
+			Describe(TEXT("String"), [this]
+			{
+				It(TEXT("loads 2 strings"), [this]
+				{
+					static constexpr const TCHAR* ExampleJson = TEXT(R"_JSON(
+					{
+						"arrayOfStrings": [
+					        "test string",
+					        "string test"
+					    ]
+					})_JSON");
+					const bool bDeserialized = DeserializeJsonString(ExampleJson);
+					TestTrue(TEXT("Deserialized Json"), bDeserialized);
+
+					if(bDeserialized)
+					{
+						Testing::FTestClass TestObject;
+						const bool bGotField = UE::Json::TryGetField(
+							JsonObject,
+							GET_MEMBER_NAME_STRING_CHECKED(Testing::FTestClass, ArrayOfStrings),
+							TestObject.ArrayOfStrings);
+								
+						TestTrue(TEXT("TryGetField"), bGotField);
+						TestEqual(TEXT("Array has 2 elements"), TestObject.ArrayOfStrings.Num(), 2);
+						TestEqual(TEXT("Array[0] is test string"), TestObject.ArrayOfStrings[0], TEXT("test string"));
+						TestEqual(TEXT("Array[1] is string test"), TestObject.ArrayOfStrings[1], TEXT("string test")); 
+					}
+				});
+			});
+
+			Describe(TEXT("Value"), [this]
+			{
+				It(TEXT("loads 3 elements"), [this]
+				{
+					static constexpr const TCHAR* ExampleJson = TEXT(R"_JSON(
+					{
+					    "arrayOfStructs": [
+					        {
+					            "testFloat": 35.5            
+					        },
+					        {
+					            "testText": "Some other text"         
+					        },
+					        {
+					            "testPtrOfContainingType": {
+					                "testFloat": 23.1            
+					            }
+					        }
+					    ]
+					})_JSON");
+					const bool bDeserialized = DeserializeJsonString(ExampleJson);
+					TestTrue(TEXT("Deserialized Json"), bDeserialized);
+
+					if(bDeserialized)
+					{
+						Testing::FTestClass TestObject;
+						const bool bGotField = UE::Json::TryGetField(
+							JsonObject,
+							GET_MEMBER_NAME_STRING_CHECKED(Testing::FTestClass, ArrayOfStructs),
+							TestObject.ArrayOfStructs);
+											
+						TestTrue(TEXT("TryGetField"), bGotField);
+						TestEqual(TEXT("Array has 3 elements"), TestObject.ArrayOfStructs.Num(), 3);
+						TestEqual(TEXT("Array[0].testFloat is 35.5"), TestObject.ArrayOfStructs[0].TestFloat, 35.5f);
+						TestEqual(TEXT("Array[1].testText is Some other text"), TestObject.ArrayOfStructs[1].TestText.ToString(), TEXT("Some other text"));
+						TestEqual(TEXT("Array[2].testPtrOfContainingType.testFloat is 23.1"), TestObject.ArrayOfStructs[2].TestPtrOfContainingType->TestFloat, 23.1f);
+					}
+				});
+			});
+
+			Describe(TEXT("SharedPtr"), [this]
+			{
+				It(TEXT("loads 3 elements"), [this]
+				{
+					static constexpr const TCHAR* ExampleJson = TEXT(R"_JSON(
+					{
+					    "arrayOfSharedPtr": [
+					        {
+					            "testFloat": 5.5            
+					        },
+					        {
+					            "testText": "SharedPtr text"         
+					        },
+					        null
+					    ]
+					})_JSON");
+					const bool bDeserialized = DeserializeJsonString(ExampleJson);
+					TestTrue(TEXT("Deserialized Json"), bDeserialized);
+
+					if(bDeserialized)
+					{
+						Testing::FTestClass TestObject;
+						const bool bGotField = UE::Json::TryGetField(
+							JsonObject,
+							GET_MEMBER_NAME_STRING_CHECKED(Testing::FTestClass, ArrayOfSharedPtr),
+							TestObject.ArrayOfSharedPtr);
+	
+						TestTrue(TEXT("TryGetField"), bGotField);
+						TestEqual(TEXT("Array has 3 elements"), TestObject.ArrayOfSharedPtr.Num(), 3);
+						TestEqual(TEXT("Array[0].testFloat"), TestObject.ArrayOfSharedPtr[0]->TestFloat, 5.5f);
+						TestEqual(TEXT("Array[1].testText is SharedPtr text"), TestObject.ArrayOfSharedPtr[1]->TestText.ToString(), TEXT("SharedPtr text"));
+						TestNull(TEXT("Array[2]"), TestObject.ArrayOfSharedPtr[2].Get());
+					}
+				});
+
+				It(TEXT("loads 4 Optional elements"), [this]
+				{
+					static constexpr const TCHAR* ExampleJson = TEXT(R"_JSON(
+					{
+					    "arrayOfSharedPtrOptional": [
+							null,
+					        {
+					            "testFloat": 5.5            
+					        },
+					        {
+					            "testText": "SharedPtr Optional text"         
+					        },
+					        null
+					    ]
+					})_JSON");
+					const bool bDeserialized = DeserializeJsonString(ExampleJson);
+					TestTrue(TEXT("Deserialized Json"), bDeserialized);
+
+					if(bDeserialized)
+					{
+						Testing::FTestClass TestObject;					
+						const bool bGotField = UE::Json::TryGetField(
+							JsonObject,
+							GET_MEMBER_NAME_STRING_CHECKED(Testing::FTestClass, ArrayOfSharedPtrOptional),
+							TestObject.ArrayOfSharedPtrOptional);
+							
+						TestTrue(TEXT("TryGetField"), bGotField);
+						TestEqual(TEXT("Array has 4 elements"), TestObject.ArrayOfSharedPtrOptional.Num(), 4);
+						TestNull(TEXT("Array[0]"), TestObject.ArrayOfSharedPtrOptional[0].Get());
+						TestEqual(TEXT("Array[1].testFloat"), TestObject.ArrayOfSharedPtrOptional[1]->GetValue().TestFloat, 5.5f);
+						TestEqual(TEXT("Array[2].testText is SharedPtr Optional text"), TestObject.ArrayOfSharedPtrOptional[2]->GetValue().TestText.ToString(), TEXT("SharedPtr Optional text"));
+						TestNull(TEXT("Array[3]"), TestObject.ArrayOfSharedPtrOptional[3].Get());
+					}
+				});
+			});
+		});
+
+		Describe(TEXT("Variant"), [this]
+		{
+			It(TEXT("loads enum"), [this]
+			{
+				static constexpr const TCHAR* ExampleJson = TEXT(R"_JSON(
+				{
+				    "Variant": "enumValueC"
+				})_JSON");
+				const bool bDeserialized = DeserializeJsonString(ExampleJson);
+				TestTrue(TEXT("Deserialized Json"), bDeserialized);
+
+				if(bDeserialized)
+				{
+					Testing::FTestClass TestObject;
+					const bool bGotField = UE::Json::TryGetField(
+						JsonObject,
+						GET_MEMBER_NAME_STRING_CHECKED(Testing::FTestClass, Variant),
+						TestObject.Variant);
+
+					if(TestTrue(TEXT("TryGetField"), bGotField))
+					{
+						TestTrue(TEXT("Variant is an enum"), TestObject.Variant.IsType<EWAPITestUEnum>());
+						TestEqual(TEXT("Enum value is EnumValueC"), TestObject.Variant.Get<EWAPITestUEnum>(), EWAPITestUEnum::EnumValueC);
+					}
+				}
+			});
+
+			It(TEXT("loads struct"), [this]
+			{
+				static constexpr const TCHAR* ExampleJson = TEXT(R"_JSON(
+				{
+				    "Variant": {
+						"testFloat": 16.5,
+						"testText": "Variant nested text"
+					}
+				})_JSON");
+				const bool bDeserialized = DeserializeJsonString(ExampleJson);
+				TestTrue(TEXT("Deserialized Json"), bDeserialized);
+
+				if(bDeserialized)
+				{
+					Testing::FTestClass TestObject;
+					const bool bGotField = UE::Json::TryGetField(
+						JsonObject,
+						GET_MEMBER_NAME_STRING_CHECKED(Testing::FTestClass, Variant),
+						TestObject.Variant);
+
+					if(TestTrue(TEXT("TryGetField"), bGotField))
+					{
+						TestTrue(TEXT("Variant is a struct"), TestObject.Variant.IsType<Testing::FTestStruct>());
+						TestEqual(TEXT("Variant.testFloat is 16.5"), TestObject.Variant.Get<Testing::FTestStruct>().TestFloat, 16.5f);
+						TestEqual(TEXT("Variant.testText is Variant nested text"), TestObject.Variant.Get<Testing::FTestStruct>().TestText.ToString(), TEXT("Variant nested text"));
+					}
+				}
+			});
+
+			It(TEXT("loads boolean"), [this]
+			{
+				static constexpr const TCHAR* ExampleJson = TEXT(R"_JSON(
+				{
+				    "Variant": false
+				})_JSON");
+				const bool bDeserialized = DeserializeJsonString(ExampleJson);
+				TestTrue(TEXT("Deserialized Json"), bDeserialized);
+
+				if(bDeserialized)
+				{
+					Testing::FTestClass TestObject;
+					const bool bGotField = UE::Json::TryGetField(
+						JsonObject,
+						GET_MEMBER_NAME_STRING_CHECKED(Testing::FTestClass, Variant),
+						TestObject.Variant);
+
+					if(TestTrue(TEXT("TryGetField"), bGotField))
+					{
+						TestEqual(TEXT("Field is a boolean"), JsonObject->GetField<EJson::Boolean>(TEXT("Variant"))->Type, EJson::Boolean);
+						TestTrue(TEXT("Variant is a boolean"), TestObject.Variant.IsType<bool>());
+						TestFalse(TEXT("Bool value is false"), TestObject.Variant.Get<bool>());
+					}
+				}
+			});
+		});
+	});
+}
+
+#endif
diff --git a/Engine/Plugins/Web/WebAPI/Source/WebAPIEditor/Private/WebAPIEditorLog.cpp b/Engine/Plugins/Web/WebAPI/Source/WebAPIEditor/Private/WebAPIEditorLog.cpp
--- a/Engine/Plugins/Web/WebAPI/Source/WebAPIEditor/Private/WebAPIEditorLog.cpp
+++ b/Engine/Plugins/Web/WebAPI/Source/WebAPIEditor/Private/WebAPIEditorLog.cpp
@@ -1,5 +1,6 @@
 // Copyright Epic Games, Inc. All Rights Reserved.
 
 #include "WebAPIEditorLog.h"
+#include "Logging/LogMacros.h"
 
 DEFINE_LOG_CATEGORY(LogWebAPIEditor);
diff --git a/Engine/Plugins/Web/WebAPI/Source/WebAPIEditor/Private/WebAPIEditorModule.h b/Engine/Plugins/Web/WebAPI/Source/WebAPIEditor/Private/WebAPIEditorModule.h
--- a/Engine/Plugins/Web/WebAPI/Source/WebAPIEditor/Private/WebAPIEditorModule.h
+++ b/Engine/Plugins/Web/WebAPI/Source/WebAPIEditor/Private/WebAPIEditorModule.h
@@ -56,6 +56,8 @@ public:
     /** Returns the static type registry to retrieve metadata for built-in types. */
     virtual UWebAPIStaticTypeRegistry* GetStaticTypeRegistry() const override;
 
+	virtual FOnRegisterLayoutExtensions& OnRegisterLayoutExtensions() override { return RegisterLayoutExtensions; }
+
 public:
     /** App identifier string */
     static const FName AppIdentifier;
@@ -67,7 +69,6 @@ private:
     EAssetTypeCategories::Type AssetCategory = EAssetTypeCategories::Misc;
     TWeakPtr<FWebAPIDefinitionAssetEditorToolkit> AssetEditorToolkit;
     TArray<TSharedPtr<class IAssetTypeActions>> Actions;
-	
 
     TSharedPtr<FExtensibilityManager> MenuExtensibilityManager;
     TSharedPtr<FExtensibilityManager> ToolBarExtensibilityManager;
@@ -76,4 +77,6 @@ private:
     FOnWebAPIProvidersChanged ProvidersChangedDelegate;
 
     TMap<FName, TScriptInterface<IWebAPICodeGeneratorInterface>> CodeGenerators;
+
+	FOnRegisterLayoutExtensions	RegisterLayoutExtensions;
 };
diff --git a/Engine/Plugins/Web/WebAPI/Source/WebAPIEditor/Private/WebAPIEditorStyle.cpp b/Engine/Plugins/Web/WebAPI/Source/WebAPIEditor/Private/WebAPIEditorStyle.cpp
--- a/Engine/Plugins/Web/WebAPI/Source/WebAPIEditor/Private/WebAPIEditorStyle.cpp
+++ b/Engine/Plugins/Web/WebAPI/Source/WebAPIEditor/Private/WebAPIEditorStyle.cpp
@@ -11,7 +11,7 @@
 namespace UE
 {
 	namespace WebAPI
-	{
+	{ 
 		namespace Private
 		{
 			const FVector2D Icon16x16(16.0f, 16.0f);
@@ -28,38 +28,36 @@ namespace UE
 
 TUniquePtr<FSlateStyleSet> FWebAPIEditorStyle::StyleInstance = nullptr;
 
-using namespace UE::WebAPI;
-
 // This is to fix the issue that SlateStyleMacros like IMAGE_BRUSH look for RootToContentDir but StyleSet->RootToContentDir is how this style is set up
-#define RootToContentDir Private::StyleInstance->RootToContentDir
-#define RootToCoreContentDir Private::StyleInstance->RootToCoreContentDir
+#define RootToContentDir UE::WebAPI::Private::StyleInstance->RootToContentDir
+#define RootToCoreContentDir UE::WebAPI::Private::StyleInstance->RootToCoreContentDir
 
 void FWebAPIEditorStyle::Register()
 {
-	Private::StyleInstance = MakeUnique<FSlateStyleSet>(Private::NAME_StyleName);
-	Private::StyleInstance->SetCoreContentRoot(FPaths::EngineContentDir() / TEXT("Editor/Slate"));
-	Private::StyleInstance->SetContentRoot(IPluginManager::Get().FindPlugin(TEXT("WebAPI"))->GetContentDir() / TEXT("Editor/Slate"));
+	UE::WebAPI::Private::StyleInstance = MakeUnique<FSlateStyleSet>(UE::WebAPI::Private::NAME_StyleName);
+	UE::WebAPI::Private::StyleInstance->SetCoreContentRoot(FPaths::EngineContentDir() / TEXT("Editor/Slate"));
+	UE::WebAPI::Private::StyleInstance->SetContentRoot(IPluginManager::Get().FindPlugin(TEXT("WebAPI"))->GetContentDir() / TEXT("Editor/Slate"));
 
-	Private::StyleInstance->Set("WebAPI.TreeView.LabelBackground", new CORE_BOX_BRUSH("Common/GroupBorderLight", 6.f/18.f, FStyleColors::Dropdown));
+	UE::WebAPI::Private::StyleInstance->Set("WebAPI.TreeView.LabelBackground", new CORE_BOX_BRUSH("Common/GroupBorderLight", 6.f/18.f, FStyleColors::Dropdown));
 	
-	FSlateStyleRegistry::RegisterSlateStyle(*Private::StyleInstance.Get());
+	FSlateStyleRegistry::RegisterSlateStyle(*UE::WebAPI::Private::StyleInstance.Get());
 }
 
 void FWebAPIEditorStyle::Unregister()
 {
-	FSlateStyleRegistry::UnRegisterSlateStyle(*Private::StyleInstance.Get());
-	Private::StyleInstance.Reset();
+	FSlateStyleRegistry::UnRegisterSlateStyle(*UE::WebAPI::Private::StyleInstance.Get());
+	UE::WebAPI::Private::StyleInstance.Reset();
 }
 
 FName FWebAPIEditorStyle::GetStyleSetName()
 {
-	return Private::NAME_StyleName;
+	return UE::WebAPI::Private::NAME_StyleName;
 }
 
 const ISlateStyle& FWebAPIEditorStyle::Get()
 {
-	check(Private::StyleInstance.IsValid());
-	return *Private::StyleInstance.Get();
+	check(UE::WebAPI::Private::StyleInstance.IsValid());
+	return *UE::WebAPI::Private::StyleInstance.Get();
 }
 
 #undef IMAGE_BRUSH
diff --git a/Engine/Plugins/Web/WebAPI/Source/WebAPIEditor/Public/CodeGen/Dom/WebAPICodeGenFile.h b/Engine/Plugins/Web/WebAPI/Source/WebAPIEditor/Public/CodeGen/Dom/WebAPICodeGenFile.h
--- a/Engine/Plugins/Web/WebAPI/Source/WebAPIEditor/Public/CodeGen/Dom/WebAPICodeGenFile.h
+++ b/Engine/Plugins/Web/WebAPI/Source/WebAPIEditor/Public/CodeGen/Dom/WebAPICodeGenFile.h
@@ -4,7 +4,6 @@
 
 #include "CoreMinimal.h"
 #include "WebAPICodeGenBase.h"
-#include "Dom/WebAPISchema.h"
 
 class UWebAPIProperty;
 class UWebAPIModel;
diff --git a/Engine/Plugins/Web/WebAPI/Source/WebAPIEditor/Public/Dom/WebAPIOperation.h b/Engine/Plugins/Web/WebAPI/Source/WebAPIEditor/Public/Dom/WebAPIOperation.h
--- a/Engine/Plugins/Web/WebAPI/Source/WebAPIEditor/Public/Dom/WebAPIOperation.h
+++ b/Engine/Plugins/Web/WebAPI/Source/WebAPIEditor/Public/Dom/WebAPIOperation.h
@@ -185,13 +185,15 @@ public:
 	//~ Begin IWebAPISchemaObjectInterface Interface.
 	virtual void SetNamespace(const FString& InNamespace) override;
 	virtual void Visit(TFunctionRef<void(IWebAPISchemaObjectInterface*&)> InVisitor) override;
-	void BindToTypeInfo() const;
 #if WITH_EDITOR
 	virtual void SetCodeText(const FString& InCodeText) override;
 	virtual void AppendCodeText(const FString& InCodeText) override;
 #endif
 	//~ End IWebAPISchemaObjectInterface Interface.
 
+	/** Associates this model with it's own TypeInfo. */
+	virtual void BindToTypeInfo();
+
 #if WITH_EDITOR
 	/** 
 	 * @return		Returns Valid if this object has data validation rules set up for it and the data for this object is valid. Returns Invalid if it does not pass the rules. Returns NotValidated if no rules are set for this object.
diff --git a/Engine/Plugins/Web/WebAPI/Source/WebAPIEditor/Public/IWebAPIEditorModule.h b/Engine/Plugins/Web/WebAPI/Source/WebAPIEditor/Public/IWebAPIEditorModule.h
--- a/Engine/Plugins/Web/WebAPI/Source/WebAPIEditor/Public/IWebAPIEditorModule.h
+++ b/Engine/Plugins/Web/WebAPI/Source/WebAPIEditor/Public/IWebAPIEditorModule.h
@@ -62,4 +62,7 @@ public:
 
 	/** Returns the static type registry to retrieve metadata for built-in types. */
 	virtual UWebAPIStaticTypeRegistry* GetStaticTypeRegistry() const = 0;
+
+	DECLARE_EVENT_OneParam(FUVEditorModule, FOnRegisterLayoutExtensions, FLayoutExtender&);
+	virtual FOnRegisterLayoutExtensions& OnRegisterLayoutExtensions() = 0;
 };
diff --git a/Engine/Plugins/Web/WebAPI/Source/WebAPIEditor/Public/WebAPIEditorLog.h b/Engine/Plugins/Web/WebAPI/Source/WebAPIEditor/Public/WebAPIEditorLog.h
--- a/Engine/Plugins/Web/WebAPI/Source/WebAPIEditor/Public/WebAPIEditorLog.h
+++ b/Engine/Plugins/Web/WebAPI/Source/WebAPIEditor/Public/WebAPIEditorLog.h
@@ -3,5 +3,6 @@
 #pragma once
 
 #include "CoreMinimal.h"
+#include "Logging/LogMacros.h"
 
-DECLARE_LOG_CATEGORY_EXTERN(LogWebAPIEditor, Log, All);
+WEBAPIEDITOR_API DECLARE_LOG_CATEGORY_EXTERN(LogWebAPIEditor, Log, All);
diff --git a/Engine/Plugins/Web/WebAPI/Source/WebAPIEditor/Public/WebAPIJsonUtilities.h b/Engine/Plugins/Web/WebAPI/Source/WebAPIEditor/Public/WebAPIJsonUtilities.h
--- a/Engine/Plugins/Web/WebAPI/Source/WebAPIEditor/Public/WebAPIJsonUtilities.h
+++ b/Engine/Plugins/Web/WebAPI/Source/WebAPIEditor/Public/WebAPIJsonUtilities.h
@@ -3,12 +3,10 @@
 #pragma once
 
 #include "CoreMinimal.h"
-#include "WebAPIEditorLog.h"
 #include "Algo/AnyOf.h"
 #include "Dom/JsonObject.h"
 #include "Dom/JsonValue.h"
 #include "Misc/TVariant.h"
-#include "Misc/Variant.h"
 
 namespace UE::Json
 {
@@ -263,18 +261,6 @@ namespace UE::Json
 		/** Underling object pointer. */
 		TSharedPtr<ObjectType> Object;
 	};
-		
-	// @note: ideal changes to Json (engine module)
-	// - Support TVariant (same field name, different types)
-	// - Support TOptional for Get* and TryGet*
-
-	// @todo: somehow account for TOptional or not in same call or minimal overload?
-
-	// template <typename ValueType>
-	// void As(const TSharedPtr<FJsonValue>& InJsonValue, ValueType& OutValue)
-	// {
-	// 	unimplemented();
-	// }
 
 	// Numeric
 	template <typename ValueType>
@@ -353,6 +339,14 @@ namespace UE::Json
 	template <typename KeyType, typename ValueType>
 	constexpr typename TEnableIf<TypeTraits::TIsStringLike<KeyType>::Value, bool>::Type // Key must be string
 	TryGetField(const TSharedPtr<FJsonObject>& InJsonObject, const FString& InFieldName, TMap<KeyType, ValueType>& OutValues); // Key must be string
+
+	template <typename MapType>
+	typename TEnableIf<TIsTMap<MapType>::Value, bool>::Type
+	TryGet(const TSharedPtr<FJsonValue>& InJsonValue, MapType& OutValues);
+
+	template <typename MapType>
+	constexpr typename TEnableIf<TIsTMap<MapType>::Value, bool>::Type
+	TryGetField(const TSharedPtr<FJsonObject>& InJsonObject, const FString& InFieldName, MapType& OutValues); // Key must be string
 	
 	// Object (with FromJson)
 	template <typename ValueType>
@@ -383,6 +377,21 @@ namespace UE::Json
 	template <typename ValueType>
 	bool TryGetField(const TSharedPtr<FJsonObject>& InJsonObject, const FString& InFieldName, TUniqueObj<ValueType>& OutValue);
 
+	// TJsonReference
+	template <typename ValueType>
+	bool TryGet(const TSharedPtr<FJsonValue>& InJsonValue, TJsonReference<ValueType>& OutValue);
+
+	template <typename ValueType>
+	bool TryGetField(const TSharedPtr<FJsonObject>& InJsonObject, const FString& InFieldName, TJsonReference<ValueType>& OutValue);
+
+	template <typename ValueType>
+	typename TEnableIf<TIsSame<ValueType, TJsonReference<typename ValueType::ElementType>>::Value, bool>::Type
+	TryGet(const TSharedPtr<FJsonValue>& InJsonValue, ValueType& OutValue);
+
+	template <typename ValueType>
+	typename TEnableIf<TIsSame<ValueType, TJsonReference<typename ValueType::ElementType>>::Value, bool>::Type
+	TryGetField(const TSharedPtr<FJsonObject>& InJsonObject, const FString& InFieldName, ValueType& OutValue);
+
 	// Object (with FromJson)
 	template <typename ValueType>
 	typename TEnableIf<
@@ -402,6 +411,9 @@ namespace UE::Json
 	template <typename ValueType>
 	typename TEnableIf<
 		TAnd<
+			TNot<TIsSame<ValueType, TSharedPtr<typename ValueType::ElementType>>>,
+			TNot<TIsSame<ValueType, TJsonReference<typename ValueType::ElementType>>>,
+			TNot<TIsTMap<ValueType>>,
 			TNot<TIsTArray<ValueType>>,
 			TNot<TypeTraits::TIsStringLike<ValueType>>,
 			TNot<TIsPODType<ValueType>>,
@@ -412,6 +424,9 @@ namespace UE::Json
 	template <typename ValueType>
 	typename TEnableIf<
 		TAnd<
+			TNot<TIsSame<ValueType, TSharedPtr<typename ValueType::ElementType>>>,
+			TNot<TIsSame<ValueType, TJsonReference<typename ValueType::ElementType>>>,
+			TNot<TIsTMap<ValueType>>,
 			TNot<TIsTArray<ValueType>>,
 			TNot<TypeTraits::TIsStringLike<ValueType>>,
 			TNot<TIsPODType<ValueType>>,
@@ -425,19 +440,20 @@ namespace UE::Json
 	template <typename ValueType>
 	bool TryGetField(const TSharedPtr<FJsonObject>& InJsonObject, const FString& InFieldName, TUniqueObj<ValueType>& OutValue);
 
-	// TJsonReference
+	// SharedPtr
 	template <typename ValueType>
-	bool TryGet(const TSharedPtr<FJsonValue>& InJsonValue, TJsonReference<ValueType>& OutValue);
+	bool TryGet(const TSharedPtr<FJsonValue>& InJsonValue, TSharedPtr<ValueType>& OutValue);
 
 	template <typename ValueType>
-	bool TryGetField(const TSharedPtr<FJsonObject>& InJsonObject, const FString& InFieldName, TJsonReference<ValueType>& OutValue);
+	bool TryGetField(const TSharedPtr<FJsonObject>& InJsonObject, const FString& InFieldName, TSharedPtr<ValueType>& OutValue);
 
-	// SharedPtr
 	template <typename ValueType>
-	bool TryGet(const TSharedPtr<FJsonValue>& InJsonValue, TSharedPtr<ValueType>& OutValue);
+	typename TEnableIf<TIsSame<ValueType, TSharedPtr<typename ValueType::ElementType>>::Value, bool>::Type
+	TryGet(const TSharedPtr<FJsonValue>& InJsonValue, ValueType& OutValue);
 
 	template <typename ValueType>
-	bool TryGetField(const TSharedPtr<FJsonObject>& InJsonObject, const FString& InFieldName, TSharedPtr<ValueType>& OutValue);
+	typename TEnableIf<TIsSame<ValueType, TSharedPtr<typename ValueType::ElementType>>::Value, bool>::Type
+	TryGetField(const TSharedPtr<FJsonObject>& InJsonObject, const FString& InFieldName, ValueType& OutValue);
 
 	// SharedRef
 	template <typename ValueType>
@@ -496,26 +512,16 @@ namespace UE::Json
 		template <typename ValueType>
 		bool operator()(ValueType& OutValue)
 		{
-			RunCount++;
-			if(bFoundMatch)
-			{
-				return bFoundMatch;
-			}
-
 			ValueType TempValue;
 			if(UE::Json::TryGet(JsonValue, TempValue))
 			{
 				OutValue = MoveTemp(TempValue);
-				bFoundMatch = true;
 				return true;
 			}
 
 			return false;
 		}
 
-		// @todo: for debug, remove
-		uint32 RunCount = 0;
-		bool bFoundMatch = false;
 		const TSharedPtr<FJsonValue> JsonValue;
 	};
 
@@ -530,21 +536,14 @@ namespace UE::Json
 		template <typename ValueType>
 		bool operator()(ValueType& OutValue)
 		{
-			if(bFoundMatch)
-			{
-				return bFoundMatch;
-			}
-				
 			if(UE::Json::TryGetField(JsonObject, FieldName, OutValue))
 			{
-				bFoundMatch = true;
 				return true;
 			}
 
 			return false;
 		}
 
-		bool bFoundMatch = false;
 		const TSharedPtr<FJsonObject> JsonObject;
 		const FString FieldName;
 	};
@@ -555,3 +554,5 @@ namespace UE::Json
 	template <typename... ValueTypes>
 	bool TryGetField(const TSharedPtr<FJsonObject>& InJsonObject, const FString& InFieldName, TVariant<ValueTypes...>& OutValue);
 }
+
+#include "WebAPIJsonUtilities.inl"
diff --git a/Engine/Plugins/Web/WebAPI/Source/WebAPIEditor/Public/WebAPIJsonUtilities.inl b/Engine/Plugins/Web/WebAPI/Source/WebAPIEditor/Public/WebAPIJsonUtilities.inl
--- a/Engine/Plugins/Web/WebAPI/Source/WebAPIEditor/Public/WebAPIJsonUtilities.inl
+++ b/Engine/Plugins/Web/WebAPI/Source/WebAPIEditor/Public/WebAPIJsonUtilities.inl
@@ -2,7 +2,10 @@
 
 #pragma once
 
-#include "WebAPIJsonUtilities.h"
+#include "WebAPIEditorLog.h"
+#include "Templates/IsUEnumClass.h"
+#include "UObject/Class.h"
+#include "UObject/ReflectedTypeAccessors.h"
 
 namespace UE::Json
 {
@@ -24,62 +27,118 @@ namespace UE::Json
 	// Enum
 	template <typename EnumType>
 	typename TEnableIf<TIsEnumClass<EnumType>::Value, bool>::Type
-	TryGet(const TSharedPtr<FJsonValue>& InJsonValue, EnumType& OutValue, const TMap<FString, EnumType>& InNameToValue)
+	TryGet(
+		const TSharedPtr<FJsonValue>& InJsonValue,
+		EnumType& OutValue,
+		const TMap<FString, EnumType>& InNameToValue)
 	{
-		// Try int first
-		uint8 IntValue;
-		bool bResult = InJsonValue->TryGetNumber(IntValue);
-		if(bResult)
-		{
-			OutValue = StaticCast<EnumType>(IntValue);	
-		}
-		else
+		bool bResult = false;
+
+		// Try bool first - TryGetNumber converts "false" to 0, etc. which is undesired for an enum
+		if(InJsonValue->Type != EJson::Boolean)
 		{
-			if(InNameToValue.IsEmpty())
+			// Try int only if bool conversion failed
+			uint8 IntValue;
+			bResult = InJsonValue->TryGetNumber(IntValue);
+			if(bResult)
 			{
-				UE_LOG(LogWebAPIEditor, Error, TEXT("Tried to convert enum from string, but no lookup was provided."));
-				return bResult;  
+				OutValue = StaticCast<EnumType>(IntValue);	
 			}
-
-			FString StrValue;
-			bResult = InJsonValue->TryGetString(StrValue);
-			if(bResult)
+			else
 			{
-				if(const EnumType* NameToValue = InNameToValue.Find(StrValue))
+				static const UEnum* TypeEnum = nullptr;
+
+				if constexpr (TIsUEnumClass<EnumType>::Value)
+				{
+					TypeEnum = StaticEnum<EnumType>();
+				}
+				
+				if(!TypeEnum && InNameToValue.IsEmpty())
 				{
-					OutValue = *NameToValue;
+					// UE_LOG(LogWebAPIEditor, Error, TEXT("Tried to convert enum from string, but no lookup was provided."));
+					return bResult;
+				}
+
+				FString StrValue;
+				bResult = InJsonValue->TryGetString(StrValue);
+				if(bResult)
+				{
+					if(TypeEnum)
+					{
+						if(int64 EnumValue = TypeEnum->GetValueByNameString(StrValue);
+							EnumValue != INDEX_NONE)
+						{
+							OutValue = StaticCast<EnumType>(EnumValue);
+						}
+					}
+					else if(const EnumType* NameToValue = InNameToValue.Find(StrValue))
+					{
+						OutValue = *NameToValue;
+					}
+					else
+					{
+						// Got a valid string, but it wasn't in the NameToValue map
+						bResult = false;
+					}
 				}
 			}
-		}			
+		}
 		return bResult;
 	}
 
 	template <typename EnumType>
 	constexpr typename TEnableIf<TIsEnumClass<EnumType>::Value, bool>::Type
-	TryGetField(const TSharedPtr<FJsonObject>& InJsonObject, const FString& InFieldName, EnumType& OutValue, const TMap<FString, EnumType>& InNameToValue)
+	TryGetField(
+		const TSharedPtr<FJsonObject>& InJsonObject,
+		const FString& InFieldName,
+		EnumType& OutValue,
+		const TMap<FString, EnumType>& InNameToValue)
 	{
 		if(const TSharedPtr<FJsonValue> JsonField = InJsonObject->TryGetField(InFieldName))
 		{
-			return TryGet(JsonField, OutValue.Get(), InNameToValue);
+			return TryGet(JsonField, OutValue, InNameToValue);
 		}
 		return false;
 	}
 
 	// String
-	template <typename ValueType>
-	constexpr typename TEnableIf<TypeTraits::TIsStringLike<ValueType>::Value, bool>::Type
-	TryGet(const TSharedPtr<FJsonValue>& InJsonValue, ValueType& OutValue)
+	template <>
+	inline bool TryGet<FString>(const TSharedPtr<FJsonValue>& InJsonValue, FString& OutValue)
 	{
 		return InJsonValue->TryGetString(OutValue);
 	}
 
-	template <typename ValueType>
-	constexpr typename TEnableIf<TypeTraits::TIsStringLike<ValueType>::Value, bool>::Type
-	TryGetField(const TSharedPtr<FJsonObject>& InJsonObject, const FString& InFieldName, ValueType& OutValue)
+	template <>
+	inline bool TryGetField<FString>(const TSharedPtr<FJsonObject>& InJsonObject, const FString& InFieldName, FString& OutValue)
 	{
 		return InJsonObject->TryGetStringField(InFieldName, OutValue);
 	}
 
+	// Text
+	template <>
+	inline bool TryGet<FText>(const TSharedPtr<FJsonValue>& InJsonValue, FText& OutValue)
+	{
+		FString StrValue;
+		if(InJsonValue->TryGetString(StrValue))
+		{
+			OutValue = FText::FromString(StrValue);
+			return true;
+		}
+		return false;
+	}
+
+	template <>
+	inline bool TryGetField<FText>(const TSharedPtr<FJsonObject>& InJsonObject, const FString& InFieldName, FText& OutValue)
+	{
+		FString StrValue;
+		if(InJsonObject->TryGetStringField(InFieldName, StrValue))
+		{
+			OutValue = FText::FromString(StrValue);
+			return true;
+		}
+		return false;
+	}
+
 	// Bool
 	template <typename ValueType>
 	constexpr typename TEnableIf<TIsSame<ValueType, bool>::Value, bool>::Type
@@ -161,13 +220,33 @@ namespace UE::Json
 		return false;
 	}
 
+	template <typename MapType>
+	typename TEnableIf<TIsTMap<MapType>::Value, bool>::Type
+	TryGet(const TSharedPtr<FJsonValue>& InJsonValue, MapType& OutValues)
+	{
+		using KeyType = typename MapType::KeyType;
+		using ValueType = typename MapType::ValueType;
+
+		return TryGet<KeyType, ValueType>(InJsonValue, OutValues);
+	}
+
+	template <typename MapType>
+	constexpr typename TEnableIf<TIsTMap<MapType>::Value, bool>::Type
+	TryGetField(const TSharedPtr<FJsonObject>& InJsonObject, const FString& InFieldName, MapType& OutValues)
+	{
+		using KeyType = typename MapType::KeyType;
+		using ValueType = typename MapType::ValueType;
+
+		return TryGetField<KeyType, ValueType>(InJsonObject, InFieldName, OutValues);
+	}
+
 	// Optional
 	template <typename ValueType>
 	constexpr typename TEnableIf<!TIsEnumClass<ValueType>::Value, bool>::Type
 	TryGet(const TSharedPtr<FJsonValue>& InJsonValue, TOptional<ValueType>& OutValue)
 	{
 		ValueType Value;
-		if(TryGet<TDecay<ValueType>>(InJsonValue, Value))
+		if(TryGet<ValueType>(InJsonValue, Value))
 		{
 			OutValue = Value;
 			return true;
@@ -181,12 +260,12 @@ namespace UE::Json
 	{
 		if(const TSharedPtr<FJsonValue> JsonField = InJsonObject->TryGetField(InFieldName))
 		{
+			// Field found, so initialize TOptional
 			if(!OutValue.IsSet())
 			{
 				OutValue = ValueType{};
 			}
-				
-			return TryGet(JsonField, OutValue.GetValue());
+			return TryGet(JsonField, OutValue);
 		}
 		return true;
 	}
@@ -255,6 +334,9 @@ namespace UE::Json
 	template <typename ValueType>
 	typename TEnableIf<
 		TAnd<
+			TNot<TIsSame<ValueType, TSharedPtr<typename ValueType::ElementType>>>,
+			TNot<TIsSame<ValueType, TJsonReference<typename ValueType::ElementType>>>,
+			TNot<TIsTMap<ValueType>>,
 			TNot<TIsTArray<ValueType>>,
 			TNot<TypeTraits::TIsStringLike<ValueType>>,
 			TNot<TIsPODType<ValueType>>,
@@ -265,10 +347,12 @@ namespace UE::Json
 		return false;
 	}
 
-	// TIsTMap<ContainerType>,
 	template <typename ValueType>
 	typename TEnableIf<
 		TAnd<
+			TNot<TIsSame<ValueType, TSharedPtr<typename ValueType::ElementType>>>,
+			TNot<TIsSame<ValueType, TJsonReference<typename ValueType::ElementType>>>,
+			TNot<TIsTMap<ValueType>>,
 			TNot<TIsTArray<ValueType>>,
 			TNot<TypeTraits::TIsStringLike<ValueType>>,
 			TNot<TIsPODType<ValueType>>,
@@ -323,10 +407,29 @@ namespace UE::Json
 		return false;
 	}
 
+	template <typename ValueType>
+	typename TEnableIf<TIsSame<ValueType, TJsonReference<typename ValueType::ElementType>>::Value, bool>::Type
+	TryGet(const TSharedPtr<FJsonValue>& InJsonValue, ValueType& OutValue)
+	{
+		return TryGet<typename ValueType::ElementType>(InJsonValue, OutValue);
+	}
+
+	template <typename ValueType>
+	typename TEnableIf<TIsSame<ValueType, TJsonReference<typename ValueType::ElementType>>::Value, bool>::Type
+	TryGetField(const TSharedPtr<FJsonObject>& InJsonObject, const FString& InFieldName, ValueType& OutValue)
+	{
+		return TryGet<typename ValueType::ElementType>(InJsonObject, InFieldName, OutValue);
+	}
+
 	// SharedPtr
 	template <typename ValueType>
 	bool TryGet(const TSharedPtr<FJsonValue>& InJsonValue,	TSharedPtr<ValueType>& OutValue)
 	{
+		if(InJsonValue->IsNull())
+		{
+			return true;
+		}
+		
 		// Initialize if invalid
 		if(!OutValue.IsValid())
 		{
@@ -346,6 +449,20 @@ namespace UE::Json
 		return false;
 	}
 
+	template <typename ValueType>
+	typename TEnableIf<TIsSame<ValueType, TSharedPtr<typename ValueType::ElementType>>::Value, bool>::Type
+	TryGet(const TSharedPtr<FJsonValue>& InJsonValue, ValueType& OutValue)
+	{
+		return TryGet<typename ValueType::ElementType>(InJsonValue, OutValue);
+	}
+
+	template <typename ValueType>
+	typename TEnableIf<TIsSame<ValueType, TSharedPtr<typename ValueType::ElementType>>::Value, bool>::Type
+	TryGetField(const TSharedPtr<FJsonObject>& InJsonObject, const FString& InFieldName, ValueType& OutValue)
+	{
+		return TryGet<typename ValueType::ElementType>(InJsonObject, InFieldName, OutValue);
+	}
+
 	// SharedRef
 	template <typename ValueType>
 	bool TryGet(const TSharedPtr<FJsonValue>& InJsonValue, TSharedRef<ValueType>& OutValue)
@@ -445,17 +562,38 @@ namespace UE::Json
 		return false;
 	}
 
+	template <size_t TypeIndex = 0, typename ...ValueTypes>
+	bool TryParse(const TSharedPtr<FJsonValue>& InJsonValue, TVariant<ValueTypes...>& OutValue)
+	{
+		using VariantType = TVariant<ValueTypes...>;
+		if constexpr (TypeIndex < TVariantSize<VariantType>::Value)
+		{
+			using ValueType = typename TNthTypeFromParameterPack<TypeIndex, ValueTypes...>::Type;
+			ValueType Value;
+			if (UE::Json::TryGet(InJsonValue, Value))
+			{
+				OutValue.template Set<ValueType>(MoveTemp(Value));
+				return true;
+			}
+
+			return TryParse<TypeIndex + 1, ValueTypes...>(InJsonValue, OutValue);
+		}
+
+		return false;
+	}
+
 	// Variant
-	template <typename ... ValueTypes>
+	template <typename... ValueTypes>
 	bool TryGet(const TSharedPtr<FJsonValue>& InJsonValue, TVariant<ValueTypes...>& OutValue)
 	{
-		FVariantValueVisitor Visitor(InJsonValue);
-
-		// @todo: asked Stefan Boberg about this, only ever tests a single type
-		return Visit(Visitor, OutValue);
+		if(TryParse<0, ValueTypes...>(InJsonValue, OutValue))
+		{
+			return true;
+		}
+		return false;
 	}
 
-	template <typename ... ValueTypes>
+	template <typename... ValueTypes>
 	bool TryGetField(const TSharedPtr<FJsonObject>& InJsonObject, const FString& InFieldName, TVariant<ValueTypes...>& OutValue)
 	{
 		if(const TSharedPtr<FJsonValue> JsonField = InJsonObject->TryGetField(InFieldName))
diff --git a/Engine/Plugins/Web/WebAPI/Source/WebAPILiquidJS/Private/WebAPILiquidJSCodeGenerator.cpp b/Engine/Plugins/Web/WebAPI/Source/WebAPILiquidJS/Private/WebAPILiquidJSCodeGenerator.cpp
--- a/Engine/Plugins/Web/WebAPI/Source/WebAPILiquidJS/Private/WebAPILiquidJSCodeGenerator.cpp
+++ b/Engine/Plugins/Web/WebAPI/Source/WebAPILiquidJS/Private/WebAPILiquidJSCodeGenerator.cpp
@@ -395,15 +395,10 @@ namespace UE
 						JsonObject->SetNumberField(TEXT("responseCode"), InCodeGenObject->ResponseCode);
 						JsonObject->SetStringField(TEXT("message"), InCodeGenObject->Message);
 
-						// @todo remove if
-						if(InCodeGenObject->Base.IsValid())
-						{
 						checkf(InCodeGenObject->Base.IsValid(), TEXT("OperationResponse should have a base type set."));
 
-
-							const TSharedPtr<FJsonObject> BaseJson = ToJson(InCodeGenObject->Base);
-							JsonObject->SetObjectField(TEXT("base"), BaseJson);
-						}
+						const TSharedPtr<FJsonObject> BaseJson = ToJson(InCodeGenObject->Base);
+						JsonObject->SetObjectField(TEXT("base"), BaseJson);
 
 						return JsonObject;
 					}
@@ -711,9 +706,6 @@ TFuture<bool> UWebAPILiquidJSCodeGenerator::IsAvailable()
 		{
 			if(!bHasRetried.load(std::memory_order_relaxed))
 			{
-				// @todo: what's this? can I use it for retry system?
-				// Request->OnRequestWillRetry()
-
 				// Try starting
 				bool bRetryResult = FModuleManager::Get().GetModulePtr<FWebAPILiquidJSModule>(TEXT("WebAPILiquidJS"))->TryStartWebApp();
 				bHasRetried.store(true, std::memory_order_relaxed);
diff --git a/Engine/Plugins/Web/WebAPI/Source/WebAPIOpenAPI/Private/Factories/WebAPIOpenAPIFactoryBase.cpp b/Engine/Plugins/Web/WebAPI/Source/WebAPIOpenAPI/Private/Factories/WebAPIOpenAPIFactoryBase.cpp
--- a/Engine/Plugins/Web/WebAPI/Source/WebAPIOpenAPI/Private/Factories/WebAPIOpenAPIFactoryBase.cpp
+++ b/Engine/Plugins/Web/WebAPI/Source/WebAPIOpenAPI/Private/Factories/WebAPIOpenAPIFactoryBase.cpp
@@ -12,17 +12,6 @@ UWebAPIOpenAPIFactoryBase::UWebAPIOpenAPIFactoryBase()
 	Formats.Add(TEXT("json;JavaScript Object Notation"));
 }
 
-TFuture<bool> UWebAPIOpenAPIFactoryBase::ImportWebAPI(UWebAPIDefinition* InDefinition, const FString& InFileName, const FString& InFileContents)
-{
-	// @todo: modify definition here
-	// write raw spec contents
-	// parse spec
-
-	// @todo: return false if json couldn't be parsed
-
-	return MakeFulfilledPromise<bool>(false).GetFuture();
-}
-
 bool UWebAPIOpenAPIFactoryBase::IsValidFileExtension(const FString& InFileExtension) const
 {
 	return InFileExtension.Equals(JsonFileType.ToString(), ESearchCase::IgnoreCase);
diff --git a/Engine/Plugins/Web/WebAPI/Source/WebAPIOpenAPI/Private/Factories/WebAPIOpenAPIFactoryBase.h b/Engine/Plugins/Web/WebAPI/Source/WebAPIOpenAPI/Private/Factories/WebAPIOpenAPIFactoryBase.h
--- a/Engine/Plugins/Web/WebAPI/Source/WebAPIOpenAPI/Private/Factories/WebAPIOpenAPIFactoryBase.h
+++ b/Engine/Plugins/Web/WebAPI/Source/WebAPIOpenAPI/Private/Factories/WebAPIOpenAPIFactoryBase.h
@@ -17,7 +17,9 @@ class WEBAPIOPENAPI_API UWebAPIOpenAPIFactoryBase
 public:
 	UWebAPIOpenAPIFactoryBase();
 
-	virtual TFuture<bool> ImportWebAPI(UWebAPIDefinition* InDefinition, const FString& InFileName, const FString& InFileContents) override;
+	virtual TFuture<bool> ImportWebAPI(UWebAPIDefinition* InDefinition, const FString& InFileName, const FString& InFileContents) override
+	PURE_VIRTUAL(UWebAPIOpenAPIFactoryBase::ImportWebAPI, return MakeFulfilledPromise<bool>(false).GetFuture();)
+	
 	virtual bool IsValidFileExtension(const FString& InFileExtension) const override;
 
 protected:
diff --git a/Engine/Plugins/Web/WebAPI/Source/WebAPIOpenAPI/Private/Tests/Samples/V3/ServerObjectExample_Variables.json b/Engine/Plugins/Web/WebAPI/Source/WebAPIOpenAPI/Private/Tests/Samples/V3/ServerObjectExample_Variables.json
--- a/Engine/Plugins/Web/WebAPI/Source/WebAPIOpenAPI/Private/Tests/Samples/V3/ServerObjectExample_Variables.json
+++ b/Engine/Plugins/Web/WebAPI/Source/WebAPIOpenAPI/Private/Tests/Samples/V3/ServerObjectExample_Variables.json
@@ -1,5 +1,4 @@
-﻿{
-  "servers": [
+﻿[
     {
       "url": "https://{username}.gigantic-server.com:{port}/{basePath}",
       "description": "The production API server",
@@ -21,4 +20,3 @@
       }
     }
   ]
-}
\ No newline at end of file
diff --git a/Engine/Plugins/Web/WebAPI/Source/WebAPIOpenAPI/Private/V2/WebAPISwaggerConverter.cpp b/Engine/Plugins/Web/WebAPI/Source/WebAPIOpenAPI/Private/V2/WebAPISwaggerConverter.cpp
--- a/Engine/Plugins/Web/WebAPI/Source/WebAPIOpenAPI/Private/V2/WebAPISwaggerConverter.cpp
+++ b/Engine/Plugins/Web/WebAPI/Source/WebAPIOpenAPI/Private/V2/WebAPISwaggerConverter.cpp
@@ -17,9 +17,6 @@
 #include "Dom/WebAPITypeRegistry.h"
 #include "V2/WebAPISwaggerSchema.h"
 
-// @todo: remove
-PRAGMA_DISABLE_OPTIMIZATION
-
 #define LOCTEXT_NAMESPACE "WebAPISwaggerConverter"
 
 #define SET_OPTIONAL(SrcProperty, DstProperty)			\
@@ -114,7 +111,6 @@ namespace UE::WebAPI::Swagger
 			{TEXT("Object"), StaticTypeRegistry->Object},
 		};
 
-		// [REDACTED_PROFILE_URL]/openapi-generator/blob/5d68bd6a03f0c48e838b4fe3b98b7e30858c0373/modules/openapi-generator/src/main/java/org/openapitools/codegen/languages/CppUE4ClientCodegen.java
 		// OpenAPI type to UE type (Prefix, Name)
 		static TMap<FString, TObjectPtr<UWebAPITypeInfo>> TypeMap =
 		{
@@ -178,7 +174,6 @@ namespace UE::WebAPI::Swagger
 				// Duplicate it with the provided definition name 
 				if (!InDefinitionName.IsEmpty())
 				{
-					// @todo: what if it's an enum? how do you know at this point if it's an enum or struct?
 					// Allow prefix to be set later depending on this?
 					Result = OutputSchema->TypeRegistry->GetOrMakeGeneratedType(EWebAPISchemaType::Model,
 						InDefinitionName,
@@ -187,10 +182,9 @@ namespace UE::WebAPI::Swagger
 					Result->Prefix = TEXT("F");
 					ensure(!Result->Name.IsEmpty());
 				}
-				// otherwise duplicate it as an unnamed/partially resolved type 
 				else
 				{
-					// @todo: prevent!
+					checkNoEntry();
 					Result = Result->Duplicate(OutputSchema->TypeRegistry);
 				}
 			}
@@ -251,12 +245,10 @@ namespace UE::WebAPI::Swagger
 		const TObjectPtr<UWebAPIStaticTypeRegistry> StaticTypeRegistry = IWebAPIEditorModuleInterface::Get().GetStaticTypeRegistry();
 
 		FString DefinitionName = InDefinitionName;
-		TObjectPtr<UWebAPITypeInfo> Result;
-		
-		TSharedPtr<UE::WebAPI::OpenAPI::V2::FSchema> ItemSchema = nullptr;
+
 		if(InSchema->Items.IsSet())
 		{
-			ItemSchema = ResolveReference(InSchema->Items.GetValue(), DefinitionName);
+			const TSharedPtr<UE::WebAPI::OpenAPI::V2::FSchema> ItemSchema = ResolveReference(InSchema->Items.GetValue(), DefinitionName);
 			if(!InSchema->Items->GetPath().IsEmpty())
 			{
 				DefinitionName = InSchema->Items->GetLastPathSegment();						
@@ -264,7 +256,10 @@ namespace UE::WebAPI::Swagger
 			return ResolveType(ItemSchema, DefinitionName);
 		}
 
-		Result = ResolveType(InSchema->Type.Get(InJsonType.IsEmpty() ? TEXT("object") : InJsonType), InSchema->Format.Get(TEXT("")),	DefinitionName);
+		TObjectPtr<UWebAPITypeInfo> Result = ResolveType(
+			InSchema->Type.Get(InJsonType.IsEmpty() ? TEXT("object") : InJsonType),
+			InSchema->Format.Get(TEXT("")),
+			DefinitionName);
 		if (InSchema->Enum.IsSet() && !InSchema->Enum.GetValue().IsEmpty())
 		{
 			if(const TObjectPtr<UWebAPITypeInfo>* FoundGeneratedType = OutputSchema->TypeRegistry->FindGeneratedType(EWebAPISchemaType::Model, NameTransformer(DefinitionName)))
@@ -302,8 +297,6 @@ namespace UE::WebAPI::Swagger
 					Args.Add(TEXT("DefinitionName"), FText::FromString(DefinitionName));
 					MessageLog->LogInfo(FText::Format(LOCTEXT("CannotResolveType", "ResolveType (object) failed to find a matching type for definition \"{DefinitionName}\", creating a new one."), Args), FWebAPISwaggerProvider::LogName);
 
-					// @bug: where is the model created for this?
-					// @todo: should this always be of type model, vs parameter?
 					Result = OutputSchema->TypeRegistry->GetOrMakeGeneratedType(
 						EWebAPISchemaType::Model,
 						NameTransformer(DefinitionName),
@@ -573,8 +566,6 @@ namespace UE::WebAPI::Swagger
 		return true;
 	}
 
-
-	
 	bool FWebAPISwaggerSchemaConverter::ConvertOperationParameter(
 		const FWebAPINameVariant& InParameterName,
 		const TSharedPtr<OpenAPI::V2::FParameter>& InParameter,
@@ -618,13 +609,6 @@ namespace UE::WebAPI::Swagger
 			return false;
 		}
 
-		if(!OutParameter->Type.TypeInfo->bIsBuiltinType)
-		{
-			// @todo: make generated type here
-			auto x = OutParameter->Type;
-			// ConvertEnum(InParameter)
-		}
-
 		OutParameter->Name = FWebAPINameInfo(InParameterName.ToString(true), InParameterName.GetJsonName(), OutParameter->Type);
 
 		// Add struct as it's own model, and reference it as this properties type
@@ -673,8 +657,6 @@ namespace UE::WebAPI::Swagger
 		OutParameter->Description = InParameter->Description.Get(TEXT(""));
 		OutParameter->bIsRequired = InParameter->bRequired.Get(false);
 		OutParameter->BindToTypeInfo();
-		
-		// PatchProperty({}, InParameterName, InParameter, InDefinitionName, OutParameter);
 
 		return true;
 	}
@@ -766,7 +748,6 @@ namespace UE::WebAPI::Swagger
 				ModelTypeName.TypeInfo->Prefix = TEXT("F");
 				ModelTypeName.TypeInfo->SetNested(ParameterTypeName);
 
-				//const TObjectPtr<UWebAPIModel>& Model = OutputSchema->AddModel<UWebAPIModel>(ModelTypeName.HasTypeInfo() ? ModelTypeName.TypeInfo.Get() : nullptr);
 				PatchModel(InSrcParameter, {}, DstParameter);
 
 				const FText LogMessage = FText::FormatNamed(
@@ -778,41 +759,9 @@ namespace UE::WebAPI::Swagger
 				
 				DstParameter->Name.TypeInfo->JsonName = ParameterTypeName.GetJsonName();
 				DstParameter->Name.TypeInfo->JsonType = UWebAPIStaticTypeRegistry::ToFromJsonType;
-
-				// DstParameter->Type = Model->Name;
-				//					DstParameter->Type.TypeInfo->Model = DstParameter;
 			}
 		}
 
-
-		/**
-		const TObjectPtr<UWebAPIModel> Model = Cast<UWebAPIModel>(DstParameter);
-		if (!PatchModel(InSrcParameter, ParameterTypeName, Model))
-		{
-			return nullptr;
-		}
-
-		
-		if (!SrcParameterSchema.GetPath().IsEmpty())
-		{
-			SrcParameterDefinitionName = SrcParameterSchema.GetLastPathSegment();
-		}
-
-		if (SrcParameterSchema.IsSet() && SrcParameterSchema->Items.IsSet() && !SrcParameterSchema->Items->GetPath().IsEmpty())
-		{
-			SrcParameterDefinitionName = SrcParameterSchema->Items->GetLastPathSegment();
-		}
-
-		if (SrcParameterSchema.IsSet())
-		{
-			DstParameter->Type = ResolveType(SrcParameterSchema.GetShared(), NameTransformer(*ParameterName));
-		}
-		else
-		{
-			DstParameter->Type = ResolveType(InSrcParameter, NameTransformer(*ParameterName));
-		}
-		*/
-
 		// Special case for "body" parameters
 		if (InSrcParameter->In == TEXT("body"))
 		{
@@ -826,43 +775,14 @@ namespace UE::WebAPI::Swagger
 
 		DstParameter->BindToTypeInfo();
 
-		//DstParameter->Name = FWebAPINameInfo(NameTransformer(DstParameter->Name), DstParameter->Name.GetJsonName(), DstParameter->Type);
 		DstParameter->Description = InSrcParameter->Description.Get(TEXT(""));
 		DstParameter->bIsRequired = InSrcParameter->bRequired.Get(false);
-		
-		/*
-		if(!DstParameter->Type.TypeInfo->bIsBuiltinType)
-		{
-			if(DstParameter->Type.TypeInfo->IsEnum())
-			{
-				ConvertEnum(InSrcParameter);						
-			}
-			else
-			{
-				unimplemented();						
-			}
-		}
-		*/
-
-		/**
-		// Special case - if the name is "body", set as Body property and not Parameter
-		if (InSrcParameter->Name.Equals(TEXT("body"), ESearchCase::IgnoreCase) && !SrcParameterDefinitionName.IsEmpty())
-		{
-			// Check for existing definition
-			if (const TObjectPtr<UWebAPITypeInfo>* FoundGeneratedType = OutputSchema->TypeRegistry->FindGeneratedType(EWebAPISchemaType::Model, SrcParameterDefinitionName))
-			{
-				DstParameter->Model = Cast<UWebAPIModel>((*FoundGeneratedType)->Model.LoadSynchronous());
-				return DstParameter;
-			}
-		}
-		*/
-
+	
 		return DstParameter;
 	}
 
 	bool FWebAPISwaggerSchemaConverter::ConvertRequest(const FWebAPITypeNameVariant& InOperationName, const TSharedPtr<OpenAPI::V2::FOperation>& InOperation, const TObjectPtr<UWebAPIOperationRequest>& OutRequest)
 	{
-		// const FWebAPITypeNameVariant RequestTypeName = OutputSchema->TypeRegistry->AddUnnamedType(IWebAPIEditorModuleInterface::Get().GetStaticTypeRegistry()->Object);
 		const FWebAPITypeNameVariant RequestTypeName = OutputSchema->TypeRegistry->GetOrMakeGeneratedType(
 			EWebAPISchemaType::Model,
 			ProviderSettings.MakeRequestTypeName(InOperationName),
@@ -970,63 +890,6 @@ namespace UE::WebAPI::Swagger
 				SrcResponseSchema.GetShared(),
 				SrcPropertyDefinitionName,
 				DstProperty);
-
-			/*
-			//OutResponse->bIsArray = IsArray(InSrcParameter);
-			if (!DstProperty->Type.HasTypeInfo())
-			{
-				DstProperty->Type = ResolveType(SrcResponseSchema->Type.Get(TEXT("object")),
-					SrcResponseSchema->Format.Get(TEXT("")),
-					SrcPropertyDefinitionName,
-					SrcResponseSchema.GetShared());
-			}
-
-			
-			if(SrcPropertyDefinitionName.IsEmpty())
-			{
-				SrcPropertyDefinitionName = ProviderSettings.MakeNestedPropertyTypeName(InOperationName, PropertyName);
-			}
-
-			const FString Type = InResponse->Schema->Get()->Type.Get(TEXT("object"));
-
-			const TObjectPtr<UWebAPIModel> Model = OutResponse;
-			if (!PatchModel(SrcResponseSchema.GetShared(), ResponseTypeName, Model))
-			{
-				return false;
-			}
-
-			// Don't use ParameterTypeName->Name, it might have a Parameter specific pre/postfix.
-			FString SrcResponseDefinitionName = ResponseTypeName.ToString(true);
-
-			// Will get schema or create if it doesn't exist (but will be empty)
-			const TSharedPtr<OpenAPI::V2::FSchema> SrcParameterSchema = ResolveReference(InResponse->Schema.Get({}), SrcResponseDefinitionName, false);
-
-			const FString PropertyName = Type == TEXT("array") ? ProviderSettings.GetDefaultArrayPropertyName() : ProviderSettings.GetDefaultPropertyName();
-			// FString DefinitionName = ProviderSettings.MakeNestedPropertyTypeName(DefinitionName, PropertyName);
-			
-			const TObjectPtr<UWebAPIProperty>& DstProperty = OutResponse->Properties.Add_GetRef(NewObject<UWebAPIProperty>(OutResponse));
-			PatchProperty(OutResponse->Name,
-				PropertyName,
-				SrcParameterSchema,
-				SrcResponseDefinitionName,
-				DstProperty);
-			
-
-			if (Type == TEXT("array"))
-			{
-
-			}
-			*/
-			// @todo: Account for single item = single property, ie. type=array, items=SomeDef would be TArray<FSomeDef> Items
-		}
-
-		if (InResponse->Headers.IsSet())
-		{
-			for (const TPair<FString, TSharedPtr<OpenAPI::V2::FHeader>>& SrcHeader : InResponse->Headers.GetValue())
-			{
-				FString Key = SrcHeader.Key; // @todo: what is this?
-				// SrcHeader.Value-> // @todo: this describes return object
-			}
 		}
 
 		OutResponse->BindToTypeInfo();
@@ -1098,7 +961,6 @@ namespace UE::WebAPI::Swagger
 				// If "Default", resolves to 0, and means all other unhandled codes (similar to default in switch statement)
 				const uint32 Code = FCString::Atoi(*SrcResponse.Key);
 				const TObjectPtr<UWebAPIOperationResponse> DstResponse = DstOperation->Responses.Add_GetRef(NewObject<UWebAPIOperationResponse>(DstOperation));
-				// OutOperation->AddResponse() @todo
 				ConvertResponse(DstOperation->Name, Code, SrcResponse.Value, DstResponse);
 
 				// If success response (code 200), had no resolved properties but the operation says it returns something, then add that something as a property
@@ -1216,10 +1078,6 @@ namespace UE::WebAPI::Swagger
 			{
 				const FString& Name = Item.Key;
 				const TSharedPtr<OpenAPI::V2::FSecurityScheme>& SecuritySchema = Item.Value;
-
-				// @todo: WebAPI security primitive
-				// TObjectPtr<UWebAPIModel> Model = OutSchema->Client->Models.Add_GetRef(NewObject<UWebAPIModel>(OutSchema));
-				// PatchModel(Name, Schema.Get(), Model);
 			}
 
 			return true;
@@ -1341,13 +1199,6 @@ namespace UE::WebAPI::Swagger
 						Operation->Rename(*OperationObjectName.ToString(), Operation->GetOuter());
 
 						Operation->BindToTypeInfo();
-
-						// Verb->Parameters
-						// Verb->Responses
-						// Verb->Schemes
-						// Verb->Consumes ie. "application/json"
-						// Verb->Produces ie. "application/json"
-						// Verb->Security
 					}
 				}
 			}
@@ -1360,6 +1211,3 @@ namespace UE::WebAPI::Swagger
 #undef SET_OPTIONAL
 
 #undef LOCTEXT_NAMESPACE
-
-// @todo: remove
-PRAGMA_ENABLE_OPTIMIZATION
diff --git a/Engine/Plugins/Web/WebAPI/Source/WebAPIOpenAPI/Private/V2/WebAPISwaggerConverter.inl b/Engine/Plugins/Web/WebAPI/Source/WebAPIOpenAPI/Private/V2/WebAPISwaggerConverter.inl
--- a/Engine/Plugins/Web/WebAPI/Source/WebAPIOpenAPI/Private/V2/WebAPISwaggerConverter.inl
+++ b/Engine/Plugins/Web/WebAPI/Source/WebAPIOpenAPI/Private/V2/WebAPISwaggerConverter.inl
@@ -181,33 +181,10 @@ namespace UE::WebAPI::Swagger
 						const FString& SrcPropertyName = NamePropertyPair.Key;
 						FString SrcPropertyDefinitionName = ProviderSettings.MakeNestedPropertyTypeName(*ModelTypeName.ToString(true), SrcPropertyName);
 						TSharedPtr<OpenAPI::V2::FSchema> SrcPropertyValue = ResolveReference(NamePropertyPair.Value, SrcPropertyDefinitionName);
-						/*
-						if (!SrcPropertyValue.IsSet())
-						{
-							if (TSharedPtr<OpenAPI::V2::FSchema> ResolvedSchema = ResolveReference(SrcPropertyValue, SrcPropertyDefinitionName))
-							{
-								SrcPropertyDefinitionName = SrcPropertyValue.GetLastPathSegment();
-								SrcPropertyValue.Set(ResolvedSchema);
-							}
-						}
-
-						if (!SrcPropertyValue.GetPath().IsEmpty())
-						{
-							SrcPropertyDefinitionName = SrcPropertyValue.GetLastPathSegment();
-						}
-
-						if (SrcPropertyValue.IsSet()
-							&& SrcPropertyValue->Items.IsSet()
-							&& !SrcPropertyValue->Items->GetPath().IsEmpty())
-						{
-							SrcPropertyDefinitionName = SrcPropertyValue->Items->GetLastPathSegment();
-						}
-						*/
 
 						if (SrcPropertyValue)
 						{
-							TObjectPtr<UWebAPIProperty>& DstProperty = Model->Properties.Add_GetRef(
-								NewObject<UWebAPIProperty>(Model));
+							TObjectPtr<UWebAPIProperty>& DstProperty = Model->Properties.Add_GetRef(NewObject<UWebAPIProperty>(Model));
 							PatchProperty(Model->Name,
 								FWebAPINameInfo(NameTransformer(SrcPropertyName), SrcPropertyName, ModelTypeName),
 								SrcPropertyValue,
diff --git a/Engine/Plugins/Web/WebAPI/Source/WebAPIOpenAPI/Private/V2/WebAPISwaggerSchema.cpp b/Engine/Plugins/Web/WebAPI/Source/WebAPIOpenAPI/Private/V2/WebAPISwaggerSchema.cpp
--- a/Engine/Plugins/Web/WebAPI/Source/WebAPIOpenAPI/Private/V2/WebAPISwaggerSchema.cpp
+++ b/Engine/Plugins/Web/WebAPI/Source/WebAPIOpenAPI/Private/V2/WebAPISwaggerSchema.cpp
@@ -3,29 +3,17 @@
 #include "V2/WebAPISwaggerSchema.h"
 
 #include "WebAPIJsonUtilities.h"
-#include "WebAPIJsonUtilities.inl"
 #include "Algo/ForEach.h"
 #include "Dom/JsonObject.h"
 
 namespace UE::WebAPI::OpenAPI::V2
 {
-	static TMap<FString, EJson> EJson_NameToValue = {
-		{TEXT("None"), EJson::None},
-		{TEXT("Null"), EJson::Null},
-		{TEXT("String"), EJson::String},
-		{TEXT("Boolean"), EJson::Boolean},
-		{TEXT("Number"), EJson::Number},
-		{TEXT("Integer"), EJson::Number},
-		{TEXT("Array"), EJson::Array},
-		{TEXT("Object"), EJson::Object},
-	};
-
 	bool FItems::FromJson(const TSharedRef<FJsonObject>& InJson)
 	{
 		bool bSetAField = false;
 		bSetAField |= Json::TryGetField(InJson, TEXT("type"), Type);
 		bSetAField |= Json::TryGetField(InJson, TEXT("format"), Format);
-		//bSetAField |= Json::TryGetField(InJson, TEXT("items"), Items); // @todo: fix recursion
+		bSetAField |= Json::TryGetField(InJson, TEXT("items"), Items);
 		bSetAField |= Json::TryGetField(InJson, TEXT("collectionFormat"), CollectionFormat);
 		bSetAField |= Json::TryGetField(InJson, TEXT("default"), Default);
 		bSetAField |= Json::TryGetField(InJson, TEXT("maximum"), Maximum);
@@ -138,7 +126,6 @@ namespace UE::WebAPI::OpenAPI::V2
 
 	bool FExample::FromJson(const TSharedRef<FJsonObject>& InJson)
 	{
-		// @todo: implement
 		const bool bSetAField = false;
 		return bSetAField;
 	}
@@ -155,7 +142,6 @@ namespace UE::WebAPI::OpenAPI::V2
 
 	bool FScopes::FromJson(const TSharedRef<FJsonObject>& InJson)
 	{
-		// @todo: implement
 		const bool bSetAField = false;
 		return bSetAField;
 	}
@@ -307,7 +293,7 @@ namespace UE::WebAPI::OpenAPI::V2
 		bSetAField |= Json::TryGetField(InJson, TEXT("produces"), Produces);
 		bSetAField |= Json::TryGetField(InJson, TEXT("parameters"), Parameters);
 		bSetAField |= Json::TryGetField(InJson, TEXT("responses"), Responses);
-		bSetAField |= Json::TryGetField(InJson, TEXT("schemes"), Schemes); // @todo
+		bSetAField |= Json::TryGetField(InJson, TEXT("schemes"), Schemes);
 		bSetAField |= Json::TryGetField(InJson, TEXT("security"), Security);
 		return bSetAField;
 	}
diff --git a/Engine/Plugins/Web/WebAPI/Source/WebAPIOpenAPI/Private/V3/WebAPIOpenAPIConverter.cpp b/Engine/Plugins/Web/WebAPI/Source/WebAPIOpenAPI/Private/V3/WebAPIOpenAPIConverter.cpp
--- a/Engine/Plugins/Web/WebAPI/Source/WebAPIOpenAPI/Private/V3/WebAPIOpenAPIConverter.cpp
+++ b/Engine/Plugins/Web/WebAPI/Source/WebAPIOpenAPI/Private/V3/WebAPIOpenAPIConverter.cpp
@@ -18,9 +18,6 @@
 #include "V3/WebAPIOpenAPIProvider.h"
 #include "V3/WebAPIOpenAPISchema.h"
 
-// @todo: remove
-PRAGMA_DISABLE_OPTIMIZATION
-
 #define LOCTEXT_NAMESPACE "WebAPIOpenAPIConverter"
 
 #define SET_OPTIONAL(SrcProperty, DstProperty)			\
@@ -97,7 +94,7 @@ namespace UE::WebAPI::OpenAPI
 		bSuccessfullyConverted &= InputSchema->Components.IsValid()
 								  ? ConvertModels(InputSchema->Components->Schemas, OutputSchema.Get())
 								  : true;
-		//bSuccessfullyConverted &= ConvertParameters(InputSchema->Components.Parameters, OutputSchema.Get());
+		
 		bSuccessfullyConverted &= ConvertPaths(InputSchema->Paths, OutputSchema.Get());
 
 		return bSuccessfullyConverted;
@@ -111,8 +108,7 @@ namespace UE::WebAPI::OpenAPI
 	TObjectPtr<UWebAPITypeInfo> FWebAPIOpenAPISchemaConverter::ResolveMappedType(const FString& InType)
 	{
 		const TObjectPtr<UWebAPIStaticTypeRegistry> StaticTypeRegistry = IWebAPIEditorModuleInterface::Get().GetStaticTypeRegistry();
-
-		// [REDACTED_PROFILE_URL]/openapi-generator/blob/5d68bd6a03f0c48e838b4fe3b98b7e30858c0373/modules/openapi-generator/src/main/java/org/openapitools/codegen/languages/CppUE4ClientCodegen.java
+		
 		// OpenAPI type to UE type (Prefix, Name)
 		static TMap<FString, TObjectPtr<UWebAPITypeInfo>> TypeMap =
 		{
@@ -220,19 +216,18 @@ namespace UE::WebAPI::OpenAPI
 				// Duplicate it with the provided definition name 
 				if (!InDefinitionName.IsEmpty())
 				{
-					// @todo: what if it's an enum? how do you know at this point if it's an enum or struct?
 					// Allow prefix to be set later depending on this?
-					Result = OutputSchema->TypeRegistry->GetOrMakeGeneratedType(EWebAPISchemaType::Model,
+					Result = OutputSchema->TypeRegistry->GetOrMakeGeneratedType(
+						EWebAPISchemaType::Model,
 						InDefinitionName,
 						InDefinitionName,
 						Result);
 					Result->Prefix = TEXT("F");
 					ensure(!Result->Name.IsEmpty());
 				}
-				// otherwise duplicate it as an unnamed/partially resolved type 
 				else
 				{
-					// @todo: prevent!
+					checkNoEntry();
 					Result = Result->Duplicate(OutputSchema->TypeRegistry);
 				}
 			}
@@ -299,8 +294,6 @@ namespace UE::WebAPI::OpenAPI
 		return InProperty->Name;
 	}
 
-
-
 	bool FWebAPIOpenAPISchemaConverter::ConvertProperty(const FWebAPITypeNameVariant& InModelName, const FWebAPINameVariant& InPropertyName, const FWebAPITypeNameVariant& InPropertyTypeName, const TObjectPtr<UWebAPIProperty>& OutProperty)
 	{
 		OutProperty->Type = InPropertyTypeName;
@@ -332,31 +325,13 @@ namespace UE::WebAPI::OpenAPI
 		};
 
 		OutParameter->Storage = InToStorage[InParameter->In];
-		
+
 		if (ParameterSchema.IsValid())
 		{
 			OutParameter->Type = ResolveType(ParameterSchema, InDefinitionName);
 		}
-		else
-		{
-			// @todo:
-			// OutParameter->Type = ResolveType(InParameter, InDefinitionName, GetDefaultJsonTypeForStorage(OutParameter->Storage));
-		}
-
+		
 		OutParameter->bIsArray = IsArray(InParameter);
-		if (!OutParameter->Type.HasTypeInfo())
-		{
-			if(InParameter->Schema.IsSet())
-			{
-				// @todo
-				/* 
-				OutParameter->Type = ResolveType(InParameter->Schema->Get(GetDefaultJsonTypeForStorage(OutParameter->Storage)),
-					TEXT(""), //InParameter->Format.Get(TEXT("")),
-					InDefinitionName,
-					ParameterSchema.GetShared());
-					*/
-			}
-		}
 
 		const TObjectPtr<UWebAPIModelBase> ModelBase = Cast<UWebAPIModelBase>(OutParameter);
 		if (!ConvertModelBase(InParameter, ModelBase))
@@ -364,13 +339,6 @@ namespace UE::WebAPI::OpenAPI
 			return false;
 		}
 
-		if(!OutParameter->Type.TypeInfo->bIsBuiltinType)
-		{
-			// @todo: make generated type here
-			auto x = OutParameter->Type;
-			// ConvertEnum(InParameter)
-		}
-
 		OutParameter->Name = FWebAPINameInfo(InParameterName.ToString(true), InParameterName.GetJsonName(), OutParameter->Type);
 
 		// Add struct as it's own model, and reference it as this properties type
@@ -419,8 +387,6 @@ namespace UE::WebAPI::OpenAPI
 		OutParameter->Description = InParameter->Description.Get(TEXT(""));
 		OutParameter->bIsRequired = InParameter->bRequired.Get(false);
 		OutParameter->BindToTypeInfo();
-		
-		// ConvertProperty({}, InParameterName, InParameter, InDefinitionName, OutParameter);
 
 		return true;
 	}
@@ -471,18 +437,7 @@ namespace UE::WebAPI::OpenAPI
 		};
 
 		DstParameter->Storage = InToStorage[InSrcParameter->In];
-
 		DstParameter->bIsArray = IsArray(InSrcParameter);
-		if (!DstParameter->Type.HasTypeInfo())
-		{
-			// @todo:
-			/*
-			DstParameter->Type = ResolveType(InSrcParameter->Type.Get(GetDefaultJsonTypeForStorage(DstParameter->Storage)),
-				InSrcParameter->Format.Get(TEXT("")),
-				SrcParameterDefinitionName,
-				SrcParameterSchema);
-				*/
-		}
 
 		const TObjectPtr<UWebAPIModel> Model = Cast<UWebAPIModel>(DstParameter);
 		if (!ConvertModel(InSrcParameter, ParameterTypeName, Model))
@@ -541,41 +496,9 @@ namespace UE::WebAPI::OpenAPI
 				
 				DstParameter->Name.TypeInfo->JsonName = ParameterTypeName.GetJsonName();
 				DstParameter->Name.TypeInfo->JsonType = UWebAPIStaticTypeRegistry::ToFromJsonType;
-
-				// DstParameter->Type = Model->Name;
-				//					DstParameter->Type.TypeInfo->Model = DstParameter;
 			}
 		}
 
-
-		/**
-		const TObjectPtr<UWebAPIModel> Model = Cast<UWebAPIModel>(DstParameter);
-		if (!ConvertModel(InSrcParameter, ParameterTypeName, Model))
-		{
-			return nullptr;
-		}
-
-		
-		if (!SrcParameterSchema.GetPath().IsEmpty())
-		{
-			SrcParameterDefinitionName = SrcParameterSchema.GetLastPathSegment();
-		}
-
-		if (SrcParameterSchema.IsSet() && SrcParameterSchema->Items.IsSet() && !SrcParameterSchema->Items->GetPath().IsEmpty())
-		{
-			SrcParameterDefinitionName = SrcParameterSchema->Items->GetLastPathSegment();
-		}
-
-		if (SrcParameterSchema.IsSet())
-		{
-			DstParameter->Type = ResolveType(SrcParameterSchema.GetShared(), NameTransformer(*ParameterName));
-		}
-		else
-		{
-			DstParameter->Type = ResolveType(InSrcParameter, NameTransformer(*ParameterName));
-		}
-		*/
-
 		// Special case for "body" parameters
 		if (InSrcParameter->In == TEXT("body"))
 		{
@@ -592,40 +515,12 @@ namespace UE::WebAPI::OpenAPI
 		//DstParameter->Name = FWebAPINameInfo(NameTransformer(DstParameter->Name), DstParameter->Name.GetJsonName(), DstParameter->Type);
 		DstParameter->Description = InSrcParameter->Description.Get(TEXT(""));
 		DstParameter->bIsRequired = InSrcParameter->bRequired.Get(false);
-		
-		/*
-		if(!DstParameter->Type.TypeInfo->bIsBuiltinType)
-		{
-			if(DstParameter->Type.TypeInfo->IsEnum())
-			{
-				ConvertEnum(InSrcParameter);						
-			}
-			else
-			{
-				unimplemented();						
-			}
-		}
-		*/
-
-		/**
-		// Special case - if the name is "body", set as Body property and not Parameter
-		if (InSrcParameter->Name.Equals(TEXT("body"), ESearchCase::IgnoreCase) && !SrcParameterDefinitionName.IsEmpty())
-		{
-			// Check for existing definition
-			if (const TObjectPtr<UWebAPITypeInfo>* FoundGeneratedType = OutputSchema->TypeRegistry->FindGeneratedType(EWebAPISchemaType::Model, SrcParameterDefinitionName))
-			{
-				DstParameter->Model = Cast<UWebAPIModel>((*FoundGeneratedType)->Model.LoadSynchronous());
-				return DstParameter;
-			}
-		}
-		*/
-
+	
 		return DstParameter;
 	}
 
 	bool FWebAPIOpenAPISchemaConverter::ConvertRequest(const FWebAPITypeNameVariant& InOperationName, const TSharedPtr<OpenAPI::V3::FOperationObject>& InOperation, const TObjectPtr<UWebAPIOperationRequest>& OutRequest)
 	{
-		// const FWebAPITypeNameVariant RequestTypeName = OutputSchema->TypeRegistry->AddUnnamedType(IWebAPIEditorModuleInterface::Get().GetStaticTypeRegistry()->Object);
 		const FWebAPITypeNameVariant RequestTypeName = OutputSchema->TypeRegistry->GetOrMakeGeneratedType(
 			EWebAPISchemaType::Model,
 			ProviderSettings.MakeRequestTypeName(InOperationName),
@@ -690,8 +585,7 @@ namespace UE::WebAPI::OpenAPI
 				{
 					FString MediaTypeDefinitionName;
 					TSharedPtr<OpenAPI::V3::FSchemaObject> MediaTypeSchema = ResolveReference(MediaType->Schema, MediaTypeDefinitionName);
-
-					// @todo: what if it's an array of $ref? TypeInfo doesn't specify whether it's an array or not - check spec
+					
 					const TObjectPtr<UWebAPIOperationParameter> DstParameter = OutRequest->Parameters.Add_GetRef(NewObject<UWebAPIOperationParameter>(OutRequest));
 					DstParameter->Type = ResolveType(MediaTypeSchema, MediaTypeDefinitionName);
 					check(DstParameter->Type.IsValid());
@@ -758,54 +652,6 @@ namespace UE::WebAPI::OpenAPI
 					SrcResponseSchema,
 					SrcPropertyDefinitionName,
 					DstProperty);
-
-				/*
-				//OutResponse->bIsArray = IsArray(InSrcParameter);
-				if (!DstProperty->Type.HasTypeInfo())
-				{
-					DstProperty->Type = ResolveType(SrcResponseSchema->Type.Get(TEXT("object")),
-						SrcResponseSchema->Format.Get(TEXT("")),
-						SrcPropertyDefinitionName,
-						SrcResponseSchema.GetShared());
-				}
-
-				
-				if(SrcPropertyDefinitionName.IsEmpty())
-				{
-					SrcPropertyDefinitionName = ProviderSettings.MakeNestedPropertyTypeName(InOperationName, PropertyName);
-				}
-
-				const FString Type = InResponse->Schema->Get()->Type.Get(TEXT("object"));
-
-				const TObjectPtr<UWebAPIModel> Model = OutResponse;
-				if (!ConvertModel(SrcResponseSchema.GetShared(), ResponseTypeName, Model))
-				{
-					return false;
-				}
-
-				// Don't use ParameterTypeName->Name, it might have a Parameter specific pre/postfix.
-				FString SrcResponseDefinitionName = ResponseTypeName.ToString(true);
-
-				// Will get schema or create if it doesn't exist (but will be empty)
-				const TSharedPtr<OpenAPI::V3::FSchemaObject> SrcParameterSchema = ResolveReference(InResponse->Schema.Get({}), SrcResponseDefinitionName, false);
-
-				const FString PropertyName = Type == TEXT("array") ? ProviderSettings.GetDefaultArrayPropertyName() : ProviderSettings.GetDefaultPropertyName();
-				// FString DefinitionName = ProviderSettings.MakeNestedPropertyTypeName(DefinitionName, PropertyName);
-				
-				const TObjectPtr<UWebAPIProperty>& DstProperty = OutResponse->Properties.Add_GetRef(NewObject<UWebAPIProperty>(OutResponse));
-				ConvertProperty(OutResponse->Name,
-					PropertyName,
-					SrcParameterSchema,
-					SrcResponseDefinitionName,
-					DstProperty);
-				
-
-				if (Type == TEXT("array"))
-				{
-
-				}
-				*/
-				// @todo: Account for single item = single property, ie. type=array, items=SomeDef would be TArray<FSomeDef> Items 
 			}
 		}
 
@@ -813,8 +659,8 @@ namespace UE::WebAPI::OpenAPI
 		{
 			for (const TPair<FString, Json::TJsonReference<OpenAPI::V3::FHeaderObject>>& SrcHeader : InResponse->Headers.GetValue())
 			{
-				FString Key = SrcHeader.Key; // @todo: what is this?
-				// SrcHeader.Value-> // @todo: this describes return object
+				FString Key = SrcHeader.Key;
+				Json::TJsonReference<OpenAPI::V3::FHeaderObject> Value = SrcHeader.Value; // @todo: handle, maybe as raw string
 			}
 		}
 
@@ -825,7 +671,6 @@ namespace UE::WebAPI::OpenAPI
 
 	TObjectPtr<UWebAPIOperation> FWebAPIOpenAPISchemaConverter::ConvertOperation(const FString& InPath, const FString& InVerb, const TSharedPtr<OpenAPI::V3::FOperationObject>& InSrcOperation, const FWebAPITypeNameVariant& InOperationTypeName)
 	{
-		
 		FString OperationName = InSrcOperation->OperationId.Get(InSrcOperation->Summary.Get(TEXT("")));
 		if(OperationName.IsEmpty())
 		{
@@ -870,7 +715,7 @@ namespace UE::WebAPI::OpenAPI
 
 		// A spec can have no tags, so ensure there's a default
 		const FString FirstTag = ProviderSettings.ToPascalCase(InSrcOperation->Tags.Get({TEXT("Default")})[0]);
-		TObjectPtr<UWebAPIService>* Service = OutputSchema->Services.Find(FirstTag);
+		const TObjectPtr<UWebAPIService>* Service = OutputSchema->Services.Find(FirstTag);
 		checkf(Service, TEXT("An operation must belong to a service!"));
 
 		const TObjectPtr<UWebAPIOperation> DstOperation = (*Service)->Operations.Add_GetRef(NewObject<UWebAPIOperation>(*Service));
@@ -887,19 +732,6 @@ namespace UE::WebAPI::OpenAPI
 
 		ConvertRequest(DstOperation->Name, InSrcOperation, DstOperation->Request);
 
-		// @todo:
-		/*
-		if (InSrcOperation->Consumes.IsSet())
-		{
-			DstOperation->RequestContentTypes.Append(InSrcOperation->Consumes.GetValue());
-		}
-
-		if (InSrcOperation->Produces.IsSet())
-		{
-			DstOperation->ResponseContentTypes.Append(InSrcOperation->Produces.GetValue());
-		}
-		*/
-
 		if (ensure(!InSrcOperation->Responses.IsEmpty()))
 		{
 			for (const TPair<FString, Json::TJsonReference<OpenAPI::V3::FResponseObject>>& SrcResponse : InSrcOperation->Responses)
@@ -907,7 +739,6 @@ namespace UE::WebAPI::OpenAPI
 				// If "Default", resolves to 0, and means all other unhandled codes (similar to default in switch statement)
 				const uint32 Code = FCString::Atoi(*SrcResponse.Key);
 				const TObjectPtr<UWebAPIOperationResponse> DstResponse = DstOperation->Responses.Add_GetRef(NewObject<UWebAPIOperationResponse>(DstOperation));
-				// OutOperation->AddResponse() @todo
 				ConvertResponse(DstOperation->Name, Code, ResolveReference(SrcResponse.Value), DstResponse);
 
 				// If success response (code 200), had no resolved properties but the operation says it returns something, then add that something as a property
@@ -982,17 +813,14 @@ namespace UE::WebAPI::OpenAPI
 			if (const TObjectPtr<UWebAPIModel>& Model = Cast<UWebAPIModel>(ModelBase))
 			{
 				Model->BindToTypeInfo();
-				// Model->Name.TypeInfo->Model = Model;
 			}
 			else if (const TObjectPtr<UWebAPIEnum>& Enum = Cast<UWebAPIEnum>(ModelBase))
 			{
 				Enum->BindToTypeInfo();
-				// Enum->Name.TypeInfo->Model = Enum;
 			}
 			else if (const TObjectPtr<UWebAPIParameter>& ParameterModel = Cast<UWebAPIParameter>(ModelBase))
 			{
 				ParameterModel->BindToTypeInfo();
-				// ParameterModel->Type.TypeInfo->Model = ParameterModel;
 			}
 		}
 
@@ -1056,7 +884,7 @@ namespace UE::WebAPI::OpenAPI
 				}
 				else if(SecurityScheme->Type == TEXT("oauth2"))
 				{
-					// @todo:
+					// @todo: handle oauth2 conversion
 					FString FlowName = SecurityScheme->Flows.IsValid() ? TEXT("") : TEXT("");
 					V3::FOAuthFlowsObject Flow = (SecurityScheme->Flows.Get())[0];
 
@@ -1072,8 +900,6 @@ namespace UE::WebAPI::OpenAPI
 					{
 						// flow = FlowName
 					}
-
-					
 				}
 
 				// @todo: WebAPI security primitive
@@ -1146,19 +972,6 @@ namespace UE::WebAPI::OpenAPI
 				}
 
 				TSharedPtr<OpenAPI::V3::FOperationObject> SrcOperation = VerbOperationPair.Value;
-				if(SrcOperation->RequestBody.IsSet())
-				{
-					FString ParameterDefinitionName;
-					// @todo:
-					/*
-					TSharedPtr<V3::FRequestBodyObject> Parameter = ResolveReference(SrcOperation->RequestBody.GetShared(), ParameterDefinitionName);
-					if(!Parameter->Content.IsEmpty())
-					{
-						auto Content = Parameter->Content;
-							
-					}
-					*/
-				}
 
 				TArray<FString> Tags;
 				if (SrcOperation.IsValid() && SrcOperation->Tags.IsSet())
@@ -1217,19 +1030,7 @@ namespace UE::WebAPI::OpenAPI
 					Operation->Rename(*OperationObjectName.ToString(), Operation->GetOuter());
 
 					Operation->BindToTypeInfo();
-
-					// Verb->Parameters
-					// Verb->Responses
-					// Verb->Schemes
-					// Verb->Consumes ie. "application/json"
-					// Verb->Produces ie. "application/json"
-					// Verb->Security							
 				}
-				// @todo:
-				// ConvertParameters(Operation)
-
-				// @todo:
-				//ConvertResponses()
 			}
 		}
 		return true;
@@ -1240,6 +1041,3 @@ namespace UE::WebAPI::OpenAPI
 #undef SET_OPTIONAL
 
 #undef LOCTEXT_NAMESPACE
-
-// @todo: remove
-PRAGMA_ENABLE_OPTIMIZATION
diff --git a/Engine/Plugins/Web/WebAPI/Source/WebAPIOpenAPI/Private/V3/WebAPIOpenAPIConverter.inl b/Engine/Plugins/Web/WebAPI/Source/WebAPIOpenAPI/Private/V3/WebAPIOpenAPIConverter.inl
--- a/Engine/Plugins/Web/WebAPI/Source/WebAPIOpenAPI/Private/V3/WebAPIOpenAPIConverter.inl
+++ b/Engine/Plugins/Web/WebAPI/Source/WebAPIOpenAPI/Private/V3/WebAPIOpenAPIConverter.inl
@@ -107,8 +107,6 @@ namespace UE::WebAPI::OpenAPI
 					Args.Add(TEXT("DefinitionName"), FText::FromString(DefinitionName));
 					MessageLog->LogInfo(FText::Format(LOCTEXT("CannotResolveType", "ResolveType (object) failed to find a matching type for definition \"{DefinitionName}\", creating a new one."), Args), FWebAPIOpenAPIProvider::LogName);
 
-					// @bug: where is the model created for this?
-					// @todo: should this always be of type model, vs parameter?
 					Result = OutputSchema->TypeRegistry->GetOrMakeGeneratedType(
 						EWebAPISchemaType::Model,
 						NameTransformer(DefinitionName),
@@ -219,7 +217,7 @@ namespace UE::WebAPI::OpenAPI
 		return nullptr;
 	}
 
-		template <typename ObjectType>
+	template <typename ObjectType>
 	TSharedPtr<ObjectType> FWebAPIOpenAPISchemaConverter::ResolveReference(const Json::TJsonReference<ObjectType>& InJsonReference, FString& OutDefinitionName, bool bInCheck)
 	{
 		if (InJsonReference.IsSet())
@@ -276,26 +274,7 @@ namespace UE::WebAPI::OpenAPI
 														const TObjectPtr<ModelType>& OutModel)
 	{
 		static_assert(TIsDerivedFrom<SchemaType, OpenAPI::V3::FSchemaObjectBase>::Value, "Type is not derived from OpenAPI::V3::FSchemaObjectBase.");
-/*
-		OutModel->Description = InSchema->Description.Get(TEXT(""));
-
-		SET_OPTIONAL(InSchema->bRequired, OutModel->bIsRequired)
-		
-		SET_OPTIONAL_FLAGGED(InSchema->Minimum, OutModel->MinimumValue, OutModel->bUseMinimumValue);
-		SET_OPTIONAL_FLAGGED(InSchema->Maximum, OutModel->MaximumValue, OutModel->bUseMaximumValue);
 
-		if(!OutModel->bUseMinimumValue)
-		{
-			SET_OPTIONAL_FLAGGED(InSchema->MinLength, OutModel->MinimumValue, OutModel->bUseMinimumValue);
-		}
-
-		if(!OutModel->bUseMaximumValue)
-		{
-			SET_OPTIONAL_FLAGGED(InSchema->MaxLength, OutModel->MaximumValue, OutModel->bUseMaximumValue);					
-		}
-		
-		SET_OPTIONAL_FLAGGED(InSchema->Pattern, OutModel->Pattern, OutModel->bUsePattern);
-		*/
 		return true;
 	}
 
diff --git a/Engine/Plugins/Web/WebAPI/Source/WebAPIOpenAPI/Private/V3/WebAPIOpenAPIFactory.cpp b/Engine/Plugins/Web/WebAPI/Source/WebAPIOpenAPI/Private/V3/WebAPIOpenAPIFactory.cpp
--- a/Engine/Plugins/Web/WebAPI/Source/WebAPIOpenAPI/Private/V3/WebAPIOpenAPIFactory.cpp
+++ b/Engine/Plugins/Web/WebAPI/Source/WebAPIOpenAPI/Private/V3/WebAPIOpenAPIFactory.cpp
@@ -38,13 +38,9 @@ TFuture<bool> UWebAPIOpenAPIFactory::ImportWebAPI(UWebAPIDefinition* InDefinitio
 	ImportData->FileContents = InFileContents;
 
 	// Parse spec
-	// @todo: move call to ConvertToWebAPISchema outside of Factory implementation (to provide Pre and Post extensibility)
 	return Provider->ConvertToWebAPISchema(InDefinition)
 	.Next([](EWebAPIConversionResult bInConversionResult)
 	{
-		// @todo: UI notification
-		//InDefinition->SetWebAPISchema(InWebAPISchema);
-
 		return bInConversionResult == EWebAPIConversionResult::Succeeded;
 	});
 }
diff --git a/Engine/Plugins/Web/WebAPI/Source/WebAPIOpenAPI/Private/V3/WebAPIOpenAPIProvider.cpp b/Engine/Plugins/Web/WebAPI/Source/WebAPIOpenAPI/Private/V3/WebAPIOpenAPIProvider.cpp
--- a/Engine/Plugins/Web/WebAPI/Source/WebAPIOpenAPI/Private/V3/WebAPIOpenAPIProvider.cpp
+++ b/Engine/Plugins/Web/WebAPI/Source/WebAPIOpenAPI/Private/V3/WebAPIOpenAPIProvider.cpp
@@ -20,7 +20,6 @@ TFuture<EWebAPIConversionResult> FWebAPIOpenAPIProvider::ConvertToWebAPISchema(c
 	
 	if(!InDefinition.IsValid())
 	{
-		// @todo: log
 		return FailedWithErrorsResult;
 	}
 
@@ -33,8 +32,6 @@ TFuture<EWebAPIConversionResult> FWebAPIOpenAPIProvider::ConvertToWebAPISchema(c
 	TSharedPtr<FJsonObject> JsonObject;
 	if(!FJsonSerializer::Deserialize(TJsonReaderFactory<TCHAR>::Create(AssetData->FileContents), JsonObject))
 	{
-		MessageLog->LogError(FText::FromString(TEXT("Couldn't deserialize file contents as Json.")), FWebAPIOpenAPIProvider::LogName);
-		//UE_LOG(LogWebAPIOpenAPI, Error, TEXT("Couldn't deserialize file contents as Json"));
 		return FailedWithErrorsResult;
 	}
 
diff --git a/Engine/Plugins/Web/WebAPI/Source/WebAPIOpenAPI/Private/V3/WebAPIOpenAPISchema.cpp b/Engine/Plugins/Web/WebAPI/Source/WebAPIOpenAPI/Private/V3/WebAPIOpenAPISchema.cpp
--- a/Engine/Plugins/Web/WebAPI/Source/WebAPIOpenAPI/Private/V3/WebAPIOpenAPISchema.cpp
+++ b/Engine/Plugins/Web/WebAPI/Source/WebAPIOpenAPI/Private/V3/WebAPIOpenAPISchema.cpp
@@ -3,7 +3,6 @@
 #include "V3/WebAPIOpenAPISchema.h"
 
 #include "WebAPIJsonUtilities.h"
-#include "WebAPIJsonUtilities.inl"
 #include "Dom/JsonObject.h"
 
 namespace UE::WebAPI::OpenAPI::V3
@@ -94,10 +93,10 @@ namespace UE::WebAPI::OpenAPI::V3
 		return bSetAField;
 	}
 
+	// @todo: parse header object
 	bool FHeaderObject::FromJson(const TSharedRef<FJsonObject>& InJson)
 	{
-		const bool bSetAField = false;
-		// @todo
+		const bool bSetAField = Super::FromJson(InJson);
 		return bSetAField;
 	}
 
@@ -253,8 +252,6 @@ namespace UE::WebAPI::OpenAPI::V3
 		bSetAField |= Json::TryGetField(InJson, TEXT("writeOnly"), bWriteOnly);
 		bSetAField |= Json::TryGetField(InJson, TEXT("xml"), Xml);
 		bSetAField |= Json::TryGetField(InJson, TEXT("externalDocs"), ExternalDocs);
-		// bSetAField |= Json::TryGetField(InJson, TEXT("example"), Example);
-
 		bSetAField |= Json::TryGetField(InJson, TEXT("properties"), Properties);
 		
 		return bSetAField;
diff --git a/Engine/Plugins/Web/WebAPI/Source/WebAPIOpenAPI/Public/V2/WebAPISwaggerSchema.h b/Engine/Plugins/Web/WebAPI/Source/WebAPIOpenAPI/Public/V2/WebAPISwaggerSchema.h
--- a/Engine/Plugins/Web/WebAPI/Source/WebAPIOpenAPI/Public/V2/WebAPISwaggerSchema.h
+++ b/Engine/Plugins/Web/WebAPI/Source/WebAPIOpenAPI/Public/V2/WebAPISwaggerSchema.h
@@ -33,9 +33,9 @@ namespace UE
 				class FItems
 				{
 				public:
-					FString Type; // will map to EJson enum, if it's string AND enum is set, it's an enum id'd by name
+					FString Type; // will map to EJson enum, if it's string AND enum is set, it's an enum by name
 					TOptional<FString> Format;
-					TOptional<TSharedPtr<FItems>> Items; // @todo: how is this handled without inf recursion?
+					TOptional<TSharedPtr<FItems>> Items;
 					TOptional<FString> CollectionFormat; // csv, ssv, tsv, pipes
 					TOptional<FString> Default;
 					TOptional<float> Maximum;
@@ -47,7 +47,7 @@ namespace UE
 					TOptional<FString> Pattern;
 					TOptional<int32> MaxItems;
 					TOptional<int32> MinItems;
-					TOptional<bool> bUniqueItems; // @todo: array should be a Set
+					TOptional<bool> bUniqueItems; // Array should be a Set (of unique values)
 					TOptional<TArray<FString>> Enum; // All possible values
 					TOptional<int32> MultipleOf;
 
@@ -95,7 +95,7 @@ namespace UE
 					TOptional<FString> Pattern;					
 					TOptional<int32> MinItems;
 					TOptional<int32> MaxItems;
-					TOptional<bool> bUniqueItems; // @todo: array should be a Set
+					TOptional<bool> bUniqueItems; // Array should be a Set (of unique values)
 					TOptional<TArray<FString>> Enum; // All possible values
 					TOptional<Json::TJsonReference<FSchema>> Items; // Item type (if this is an array)
 
@@ -161,7 +161,7 @@ namespace UE
 					TOptional<FString> Pattern;
 					TOptional<int32> MaxItems;
 					TOptional<int32> MinItems;
-					TOptional<bool> bUniqueItems; // @todo: array should be a Set
+					TOptional<bool> bUniqueItems; // Array should be a Set (of unique values)
 					TOptional<TArray<FString>> Enum;
 					TOptional<int32> MultipleOf;
 
@@ -172,7 +172,7 @@ namespace UE
 				class FExample : public TMap<FString, FString>
 				{
 				public:
-					// @todo: key is mime-type
+					// Key is mime-type
 					bool FromJson(const TSharedRef<FJsonObject>& InJson);
 				};
 
@@ -254,9 +254,9 @@ namespace UE
 					TSharedPtr<FInfo> Info;
 					TOptional<FString> Host;
 					TOptional<FString> BasePath;
-					TOptional<TArray<FString>> Schemes; // @todo: use enum
-					TOptional<TArray<FString>> Consumes; // @todo: mime types
-					TOptional<TArray<FString>> Produces; // @todo: mimetypes
+					TOptional<TArray<FString>> Schemes;
+					TOptional<TArray<FString>> Consumes; // mimetypes
+					TOptional<TArray<FString>> Produces; // mimetypes
 					TMap<FString, TSharedPtr<FPath>> Paths;
 					TOptional<TMap<FString, TSharedPtr<FSchema>>> Definitions;
 					TOptional<TMap<FString, TSharedPtr<FParameter>>> Parameters;
@@ -313,8 +313,8 @@ namespace UE
 					TOptional<FString> Description;
 					TOptional<TSharedPtr<FExternalDocumentation>> ExternalDocs;
 					TOptional<FString> OperationId;
-					TOptional<TArray<FString>> Consumes; // @todo: mimetypes
-					TOptional<TArray<FString>> Produces; // @todo: mimetypes
+					TOptional<TArray<FString>> Consumes; // mimetypes
+					TOptional<TArray<FString>> Produces; // mimetypes
 					TOptional<TArray<Json::TJsonReference<FParameter>>> Parameters;
 					TMap<FString, TSharedPtr<FResponse>>  Responses;
 					TArray<FString> Schemes; //  ('http' | 'https' | 'ws' | 'wss')[];
diff --git a/Engine/Plugins/Web/WebAPI/Source/WebAPIOpenAPI/Public/V3/WebAPIOpenAPISchema.h b/Engine/Plugins/Web/WebAPI/Source/WebAPIOpenAPI/Public/V3/WebAPIOpenAPISchema.h
--- a/Engine/Plugins/Web/WebAPI/Source/WebAPIOpenAPI/Public/V3/WebAPIOpenAPISchema.h
+++ b/Engine/Plugins/Web/WebAPI/Source/WebAPIOpenAPI/Public/V3/WebAPIOpenAPISchema.h
@@ -90,32 +90,8 @@ namespace UE
 					TOptional<TMap<FString, Json::TJsonReference<FExampleObject>>> Examples;
 
 					TOptional<TMap<FString, TSharedPtr<FMediaTypeObject>>> Content; // only one item
-					// @todo: style values
-			
-					// MAY be extended with Specification Extensions
 
-					const FString& GetStyle()
-					{
-						if(Style.IsSet())
-						{
-							return Style.GetValue();
-						}
-
-						static TMap<FString, FString> DefaultStylesForIn = {
-							{ TEXT("query"), TEXT("form") },
-							{ TEXT("path"), TEXT("simple") },
-							{ TEXT("header"), TEXT("simple") },
-							{ TEXT("cookie"), TEXT("form") },
-						};
-
-						if(const FString* FoundStyle = DefaultStylesForIn.Find(In))
-						{
-							return *FoundStyle;
-						}
-
-						static FString Empty;
-						return Empty;
-					}
+					// MAY be extended with Specification Extensions
 
 					bool FromJson(const TSharedRef<FJsonObject>& InJson);
 				};
@@ -146,7 +122,7 @@ namespace UE
 					TOptional<FString> Pattern;
 					TOptional<int32> MinItems;
 					TOptional<int32> MaxItems;
-					TOptional<bool> bUniqueItems; // @todo: array should be a Set
+					TOptional<bool> bUniqueItems; // Array should be a Set
 					TOptional<int32> MinProperties;
 					TOptional<int32> MaxProperties;
 					TOptional<TArray<FString>> Enum; // All possible values
