WebAPI uplugin details

diff --git a/Engine/Plugins/Web/WebAPI/WebAPI.uplugin b/Engine/Plugins/Web/WebAPI/WebAPI.uplugin
--- a/Engine/Plugins/Web/WebAPI/WebAPI.uplugin
+++ b/Engine/Plugins/Web/WebAPI/WebAPI.uplugin
@@ -3,10 +3,10 @@
 	"Version": 1,
 	"VersionName": "1.0",
 	"FriendlyName": "WebAPI",
-	"Description": "",
+	"Description": "Automated generation of web based API's,",
 	"Category": "Web",
-	"CreatedBy": "",
-	"CreatedByURL": "",
+	"CreatedBy" : "Epic Games, Inc.",
+	"CreatedByURL" : "http://epicgames.com",
 	"DocsURL": "",
 	"MarketplaceURL": "",
 	"SupportURL": "",
