WebAPICodeGenFile.h non-unity fix

diff --git a/Engine/Plugins/Web/WebAPI/Source/WebAPIEditor/Public/CodeGen/Dom/WebAPICodeGenFile.h b/Engine/Plugins/Web/WebAPI/Source/WebAPIEditor/Public/CodeGen/Dom/WebAPICodeGenFile.h
--- a/Engine/Plugins/Web/WebAPI/Source/WebAPIEditor/Public/CodeGen/Dom/WebAPICodeGenFile.h
+++ b/Engine/Plugins/Web/WebAPI/Source/WebAPIEditor/Public/CodeGen/Dom/WebAPICodeGenFile.h
@@ -4,6 +4,7 @@
 
 #include "CoreMinimal.h"
 #include "WebAPICodeGenBase.h"
+#include "UObject/WeakObjectPtrTemplates.h"
 
 class UWebAPIProperty;
 class UWebAPIModel;
@@ -25,7 +26,7 @@ public:
 	/** Name of the file, excluding extension. */
 	FString FileName;
 
-	/** File extension. */
+ 	/** File extension. */
 	FString FileType;
 
 	/** Copyright notice to apply to the top of the file. Uses the Project setting unless specified. */
@@ -72,7 +73,7 @@ void FWebAPICodeGenFile::AddItems(const TArray<TSharedPtr<CodeGenType>>& InItems
 {
 	static_assert(std::is_base_of_v<FWebAPICodeGenBase, CodeGenType>, "CodeGenType must be derived from FWebAPICodeGenBase");
 
-	for(const TObjectPtr<CodeGenType>& Item : InItems)
+	for(const TSharedPtr<CodeGenType>& Item : InItems)
 	{
 		AddItem(Item);
 	}
