WebAPI cleanup and npm install fix

diff --git a/Engine/Plugins/Web/WebAPI/Source/WebAPI/Private/Security/WebAPIAuthentication.cpp b/Engine/Plugins/Web/WebAPI/Source/WebAPI/Private/Security/WebAPIAuthentication.cpp
--- a/Engine/Plugins/Web/WebAPI/Source/WebAPI/Private/Security/WebAPIAuthentication.cpp
+++ b/Engine/Plugins/Web/WebAPI/Source/WebAPI/Private/Security/WebAPIAuthentication.cpp
@@ -5,6 +5,7 @@
 #include "WebAPIDeveloperSettings.h"
 #include "WebAPILog.h"
 #include "WebAPISubsystem.h"
+#include "WebAPIUtilities.h"
 #include "Engine/Engine.h"
 #include "Interfaces/IHttpRequest.h"
 #include "Serialization/JsonSerializer.h"
@@ -55,31 +56,30 @@ bool FWebAPIOAuthSchemeHandler::HandleHttpResponse(EHttpResponseCodes::Type InRe
 
 	GEngine->GetEngineSubsystem<UWebAPISubsystem>()->MakeHttpRequest(TEXT("POST"), [&, AuthSettings](const TSharedPtr<IHttpRequest>& InRequest)
 	{
-		//Microsoft login rest URL
 		FStringFormatNamedArguments UrlArgs;
-		UrlArgs.Add(TEXT("TenantId"), AuthSettings->TenantId);
+		for(const TPair<FString, FString>& KVP : AuthSettings->AdditionalRequestQueryParameters)
+		{
+			UrlArgs.Add(KVP.Key, KVP.Value);
+		}
+		
 		const FString Url = FString::Format(*AuthSettings->AuthenticationServer, UrlArgs);
 
-		// @todo: move!
+
 		FString StringPayload = "grant_type=client_credentials";
 		StringPayload.Append("&client_id=" + AuthSettings->ClientId);
 		StringPayload.Append("&client_secret=" + AuthSettings->ClientSecret);
-		StringPayload.Append("&resource=https://digitaltwins.azure.net");
+
+		for(const TPair<FString, FString>& KVP : AuthSettings->AdditionalRequestQueryParameters)
+		{
+			StringPayload.Append(FString::Printf(TEXT("&%s=%s"), *KVP.Key, *KVP.Value));
+		}
 
 		InRequest->SetURL(Url);
 		InRequest->SetContentAsString(StringPayload);
 
 		TMap<FString, FString> Headers;
-		// FInternetAddr
-		FString SchemeName;
-		FParse::SchemeNameFromURI(*Url, SchemeName);
-		FString Host = Url.Replace(*(SchemeName + TEXT("://")), TEXT(""));
-		
-		int32 DelimiterIdx = -1;
-		if(Host.FindChar(TEXT('/'), DelimiterIdx))
-		{
-			Host = Host.Left(DelimiterIdx);
-		}
+
+		const FString Host = UWebAPIUtilities::GetHostFromUrl(*Url);
 
 		Headers.Add("Host", Host);
 		Headers.Add("Content-Type", "application/x-www-form-urlencoded");
@@ -131,7 +131,7 @@ bool FWebAPIOAuthSchemeHandler::HandleHttpResponse(EHttpResponseCodes::Type InRe
 				}
 			}
 			else
-				{
+			{
 				UE_LOG(LogWebAPI, Error, TEXT("Authentication failed: Response not valid"));
 				Message = "Response code not valid:" + Message;
 			}
@@ -142,15 +142,13 @@ bool FWebAPIOAuthSchemeHandler::HandleHttpResponse(EHttpResponseCodes::Type InRe
 			Message = "Response is null";
 		}
 
-		// @todo: request retry
-		//Redo all the requests that arrived while fetching new token
-		/*
-		if (bUpdatingToken && bSuccess)
+		if(bSuccess)
 		{
-			UpdatingToken = false;
-			RedoRequests();
+			const FString Host = UWebAPIUtilities::GetHostFromUrl(InResponse.Key->GetURL());
+
+			// Auth succeeded, retry any buffered, previous denied requests
+			GEngine->GetEngineSubsystem<UWebAPISubsystem>()->RetryRequestsForHost(Host);
 		}
-		*/
 
 		return bSuccess;
 	});
diff --git a/Engine/Plugins/Web/WebAPI/Source/WebAPI/Private/WebAPISubsystem.cpp b/Engine/Plugins/Web/WebAPI/Source/WebAPI/Private/WebAPISubsystem.cpp
--- a/Engine/Plugins/Web/WebAPI/Source/WebAPI/Private/WebAPISubsystem.cpp
+++ b/Engine/Plugins/Web/WebAPI/Source/WebAPI/Private/WebAPISubsystem.cpp
@@ -6,9 +6,16 @@
 #include "HttpModule.h"
 #include "WebAPIDeveloperSettings.h"
 #include "WebAPIOperationObject.h"
+#include "WebAPIUtilities.h"
 #include "Engine/Engine.h"
 #include "Security/WebAPIAuthentication.h"
 
+TAutoConsoleVariable<int32> CVarMaxDeniedHttpRetryCount(
+	TEXT("HTTP.MaxDeniedRetryCount"),
+	128,
+	TEXT("Maximum number of http retries for denied requests."),
+	ECVF_Default);
+
 TObjectPtr<UWebAPIOperationObject> FWebAPIPooledOperation::Pop()
 {
 	if(ensureMsgf(!ItemClass.IsNull(), TEXT("ItemClass %s could not be loaded."), *ItemClass.GetAssetName()))
@@ -50,6 +57,10 @@ bool FWebAPIPooledOperation::Push(const TObjectPtr<UWebAPIOperationObject>& InIt
 	return ItemsRemoved > 0;
 }
 
+UWebAPISubsystem::UWebAPISubsystem()
+{
+}
+
 TObjectPtr<UWebAPIOperationObject> UWebAPISubsystem::MakeOperation(const UWebAPIDeveloperSettings* InSettings, const TSubclassOf<UWebAPIOperationObject>& InClass)
 {
 	check(InClass);
@@ -79,17 +90,15 @@ void UWebAPISubsystem::ReleaseOperation(const TSubclassOf<UWebAPIOperationObject
 	NamedOperationPool.Push(InOperation);
 }
 
-TFuture<TTuple<TSharedPtr<IHttpResponse, ESPMode::ThreadSafe>, bool>> UWebAPISubsystem::MakeHttpRequest(const FString& InVerb, TUniqueFunction<void(const TSharedRef<IHttpRequest, ESPMode::ThreadSafe>)> OnRequestCreated) const
+TFuture<TTuple<TSharedPtr<IHttpResponse, ESPMode::ThreadSafe>, bool>> UWebAPISubsystem::MakeHttpRequest(const FString& InVerb, TUniqueFunction<void(const TSharedRef<IHttpRequest, ESPMode::ThreadSafe>)> OnRequestCreated)
 {
 	check(!InVerb.IsEmpty());
-	
-	// @todo: buffer requests here and retry after auth
 
 	const TSharedPtr<TPromise<TTuple<FHttpResponsePtr, bool>>, ESPMode::ThreadSafe> Promise = MakeShared<TPromise<TTuple<FHttpResponsePtr, bool>>, ESPMode::ThreadSafe>();
 
 	FHttpModule& HttpModule = FHttpModule::Get();
 
-	const TSharedRef<IHttpRequest, ESPMode::ThreadSafe> HttpRequest = HttpModule.CreateRequest();
+	const TSharedRef<IHttpRequest, ESPMode::ThreadSafe>& HttpRequest = RequestBuffer.Emplace_GetRef(HttpModule.CreateRequest());
 	HttpRequest->SetVerb(InVerb);
 
 	HttpRequest->OnProcessRequestComplete().BindLambda(
@@ -101,11 +110,19 @@ TFuture<TTuple<TSharedPtr<IHttpResponse, ESPMode::ThreadSafe>, bool>> UWebAPISub
 	// Allows the caller to inject their own Request properties
 	OnRequestCreated(HttpRequest);
 
-	HttpRequest->ProcessRequest();
-
 	return Promise->GetFuture();
 }
 
+void UWebAPISubsystem::RetryRequestsForHost(const FString& InHost)
+{
+	// Move any requests found for this host to the main queue.
+	if(const TArray<TSharedRef<IHttpRequest>>* Requests = HostRequestBuffer.Find(InHost))
+	{
+		RequestBuffer.Append(*Requests);
+		HostRequestBuffer.Remove(InHost);		
+	}
+}
+
 void UWebAPISubsystem::Initialize(FSubsystemCollectionBase& Collection)
 {
 	Super::Initialize(Collection);
@@ -120,6 +137,62 @@ void UWebAPISubsystem::Deinitialize()
 	Super::Deinitialize();
 }
 
+ETickableTickType UWebAPISubsystem::GetTickableTickType() const
+{
+	return FTickableGameObject::GetTickableTickType();
+}
+
+bool UWebAPISubsystem::IsAllowedToTick() const
+{
+	return FTickableGameObject::IsAllowedToTick();
+}
+
+void UWebAPISubsystem::Tick(float DeltaTime)
+{
+	for(TSharedRef<IHttpRequest>& Request : RequestBuffer)
+	{
+		const EHttpRequestStatus::Type RequestStatus = Request->GetStatus();
+		if(RequestStatus == EHttpRequestStatus::NotStarted)
+		{
+			Request->ProcessRequest();
+		}
+		else if(RequestStatus == EHttpRequestStatus::Failed_ConnectionError)
+		{			
+			Request->ProcessRequest();
+		}
+		else if(RequestStatus == EHttpRequestStatus::Failed)
+		{
+			// If there's a response that was denied, keep for retry after auth
+			if(Request->GetResponse().IsValid())
+			{
+				if(Request->GetResponse()->GetResponseCode() == EHttpResponseCodes::Denied)
+				{
+					FString Host = UWebAPIUtilities::GetHostFromUrl(Request->GetURL());
+
+					TArray<TSharedRef<IHttpRequest>>& HostRequests = HostRequestBuffer.FindOrAdd(Host);
+					// Prevent filling this buffer due to never being authorized, etc. - discards requests if buffer full
+					if(HostRequests.Num() <= CVarMaxDeniedHttpRetryCount->GetInt())
+					{
+						HostRequests.Add(Request);
+						HostRequestBuffer.Add(Host);
+					}
+
+					RequestBuffer.Remove(Request);
+
+					continue;
+				}
+			}
+			
+			RequestBuffer.Remove(Request);			
+		}
+	}
+}
+
+TStatId UWebAPISubsystem::GetStatId() const
+{
+	RETURN_QUICK_DECLARE_CYCLE_STAT(UWebAPISubsystem, STATGROUP_Tickables);
+}
+
 bool UWebAPISubsystem::HandleHttpRequest(TSharedPtr<IHttpRequest> InRequest, UWebAPIDeveloperSettings* InSettings)
 {
 	const TArray<TSharedPtr<FWebAPIAuthenticationSchemeHandler>>& AuthenticationHandlers = InSettings->GetAuthenticationHandlers();
diff --git a/Engine/Plugins/Web/WebAPI/Source/WebAPI/Private/WebAPIUtilities.cpp b/Engine/Plugins/Web/WebAPI/Source/WebAPI/Private/WebAPIUtilities.cpp
--- a/Engine/Plugins/Web/WebAPI/Source/WebAPI/Private/WebAPIUtilities.cpp
+++ b/Engine/Plugins/Web/WebAPI/Source/WebAPI/Private/WebAPIUtilities.cpp
@@ -6,3 +6,18 @@ const FText& UWebAPIUtilities::GetResponseMessage(const FWebAPIMessageResponse&
 {
 	return MessageResponse.GetMessage();
 }
+
+FString UWebAPIUtilities::GetHostFromUrl(const FString& InUrl)
+{
+	FString SchemeName;
+	FParse::SchemeNameFromURI(*InUrl, SchemeName);
+	FString Host = InUrl.Replace(*(SchemeName + TEXT("://")), TEXT(""));
+		
+	int32 DelimiterIdx = -1;
+	if(Host.FindChar(TEXT('/'), DelimiterIdx))
+	{
+		Host = Host.Left(DelimiterIdx);
+	}
+
+	return Host;
+}
diff --git a/Engine/Plugins/Web/WebAPI/Source/WebAPI/Public/Security/WebAPIAuthentication.h b/Engine/Plugins/Web/WebAPI/Source/WebAPI/Public/Security/WebAPIAuthentication.h
--- a/Engine/Plugins/Web/WebAPI/Source/WebAPI/Public/Security/WebAPIAuthentication.h
+++ b/Engine/Plugins/Web/WebAPI/Source/WebAPI/Public/Security/WebAPIAuthentication.h
@@ -37,10 +37,6 @@ class WEBAPI_API UWebAPIOAuthSettings : public UWebAPIAuthenticationSettings
 
 public:
 	UWebAPIOAuthSettings();
-	
-	/**  */
-	UPROPERTY(Config, EditAnywhere, BlueprintReadWrite, Category = "Security")
-	FString TenantId; // API key? @todo: this is specific to Azure services! move!
 
 	/** Public client identifier. */
 	UPROPERTY(Config, EditAnywhere, BlueprintReadWrite, Category = "Security")
@@ -66,6 +62,14 @@ public:
 	UPROPERTY(Config, EditAnywhere, BlueprintReadWrite, Category = "Security")
 	FString AuthenticationServer;
 
+	/** Additional query parameters to add to auth request. Each key should be present in the Url. */
+	UPROPERTY(Config, EditAnywhere, BlueprintReadWrite, Category = "Security")
+	TMap<FString, FString> AdditionalRequestQueryParameters;
+
+	/** Additional content key/value pairs to add to auth request. */
+	UPROPERTY(Config, EditAnywhere, BlueprintReadWrite, Category = "Security")
+	TMap<FString, FString> AdditionalRequestBodyParameters;
+
 	/** Returns true if this contains sufficient information to request authentication. */
 	bool IsValid() const;
 };
diff --git a/Engine/Plugins/Web/WebAPI/Source/WebAPI/Public/WebAPIDeveloperSettings.h b/Engine/Plugins/Web/WebAPI/Source/WebAPI/Public/WebAPIDeveloperSettings.h
--- a/Engine/Plugins/Web/WebAPI/Source/WebAPI/Public/WebAPIDeveloperSettings.h
+++ b/Engine/Plugins/Web/WebAPI/Source/WebAPI/Public/WebAPIDeveloperSettings.h
@@ -80,7 +80,6 @@ protected:
 
 	TArray<TSharedPtr<FWebAPIAuthenticationSchemeHandler>> AuthenticationHandlers;
 
-	// @todo: temporary solution to request auth tokens, etc. replace with something purpose built! Move to subsystem?
 	/** Called for all responses, providing the opportunity for custom interception. */
 	virtual bool HandleHttpResponse(EHttpResponseCodes::Type InResponseCode, TSharedPtr<class IHttpResponse, ESPMode::ThreadSafe> InResponse, bool bInWasSuccessful, UWebAPIDeveloperSettings* InSettings) override;
 };
diff --git a/Engine/Plugins/Web/WebAPI/Source/WebAPI/Public/WebAPISubsystem.h b/Engine/Plugins/Web/WebAPI/Source/WebAPI/Public/WebAPISubsystem.h
--- a/Engine/Plugins/Web/WebAPI/Source/WebAPI/Public/WebAPISubsystem.h
+++ b/Engine/Plugins/Web/WebAPI/Source/WebAPI/Public/WebAPISubsystem.h
@@ -3,6 +3,7 @@
 #pragma once
 
 #include "CoreMinimal.h"
+#include "Tickable.h"
 #include "WebAPIHttpMessageHandlers.h"
 #include "Subsystems/EngineSubsystem.h"
 
@@ -42,12 +43,15 @@ struct FWebAPIPooledOperation
 UCLASS()
 class WEBAPI_API UWebAPISubsystem
 	: public UEngineSubsystem
+	, public FTickableGameObject 
 	, public FWebAPIHttpRequestHandlerInterface
 	, public FWebAPIHttpResponseHandlerInterface
 {
 	GENERATED_BODY()
 
 public:
+	UWebAPISubsystem();
+	
 	/** Get or create pooled operation for the given OperationType. */
 	TObjectPtr<UWebAPIOperationObject> MakeOperation(const UWebAPIDeveloperSettings* InSettings, const TSubclassOf<UWebAPIOperationObject>& InClass);
 
@@ -63,22 +67,39 @@ public:
 	void ReleaseOperation(const TObjectPtr<OperationType>& InOperation);
 
 	/** Convenience function to create a simple Http request. */
-	TFuture<TTuple<TSharedPtr<IHttpResponse, ESPMode::ThreadSafe>, bool>> MakeHttpRequest(const FString& InVerb, TUniqueFunction<void(const TSharedRef<class IHttpRequest, ESPMode::ThreadSafe>)> OnRequestCreated) const;
+	TFuture<TTuple<TSharedPtr<IHttpResponse, ESPMode::ThreadSafe>, bool>> MakeHttpRequest(const FString& InVerb, TUniqueFunction<void(const TSharedRef<class IHttpRequest, ESPMode::ThreadSafe>)> OnRequestCreated);
 
+	/** Retry any failed and buffered requests for the given host. */
+	void RetryRequestsForHost(const FString& InHost);
+	
 	//~Begin UEngineSubsystem interface
 	virtual void Initialize(FSubsystemCollectionBase& Collection) override;
 	virtual void Deinitialize() override;
 	//~End UEngineSubsystem interface
 	
+	// FTickableGameObject implementation Begin
+	virtual UWorld* GetTickableGameObjectWorld() const override { return GetWorld(); }
+	virtual ETickableTickType GetTickableTickType() const override;
+	virtual bool IsAllowedToTick() const override final;
+	virtual void Tick(float DeltaTime) override;
+	virtual TStatId GetStatId() const override;
+	// FTickableGameObject implementation End
+
 private:
 	friend class UWebAPIOperationObject;
-
+ 
 	UPROPERTY()
 	bool bUsePooling = false;
 	
 	UPROPERTY(Transient)
 	TMap<FName, FWebAPIPooledOperation> OperationPool;
 
+	/** Requests to process. */
+	TArray<TSharedRef<IHttpRequest>> RequestBuffer;
+
+	/** Requests retained per host for retry when response code indicates an auth error. */
+	TMap<FString, TArray<TSharedRef<IHttpRequest>>> HostRequestBuffer;
+
 	/** Return true if the request was handled, subsequent handlers won't be called. */
 	virtual bool HandleHttpRequest(TSharedPtr<IHttpRequest> InRequest, UWebAPIDeveloperSettings* InSettings) override;
 
diff --git a/Engine/Plugins/Web/WebAPI/Source/WebAPI/Public/WebAPIUtilities.h b/Engine/Plugins/Web/WebAPI/Source/WebAPI/Public/WebAPIUtilities.h
--- a/Engine/Plugins/Web/WebAPI/Source/WebAPI/Public/WebAPIUtilities.h
+++ b/Engine/Plugins/Web/WebAPI/Source/WebAPI/Public/WebAPIUtilities.h
@@ -17,6 +17,11 @@ class WEBAPI_API UWebAPIUtilities : public UBlueprintFunctionLibrary
 	GENERATED_BODY()
 
 public:
-	UFUNCTION(BlueprintCallable, BlueprintPure, Category="WebAPI", meta=(ScriptMethod))
+	/** Return the message from the provided response. */
+	UFUNCTION(BlueprintCallable, BlueprintPure, Category="WebAPI")
 	static const FText& GetResponseMessage(const FWebAPIMessageResponse& MessageResponse);
+
+	/** Return the host (only) from the provided Url. */
+	UFUNCTION(BlueprintCallable, BlueprintPure, Category="WebAPI")
+	static FString GetHostFromUrl(const FString& InUrl);
 };
diff --git a/Engine/Plugins/Web/WebAPI/Source/WebAPIBlueprintGraph/WebAPIBlueprintGraph.Build.cs b/Engine/Plugins/Web/WebAPI/Source/WebAPIBlueprintGraph/WebAPIBlueprintGraph.Build.cs
--- a/Engine/Plugins/Web/WebAPI/Source/WebAPIBlueprintGraph/WebAPIBlueprintGraph.Build.cs
+++ b/Engine/Plugins/Web/WebAPI/Source/WebAPIBlueprintGraph/WebAPIBlueprintGraph.Build.cs
@@ -9,30 +9,36 @@ public class WebAPIBlueprintGraph : ModuleRules
 		PCHUsage = ModuleRules.PCHUsageMode.UseExplicitOrSharedPCHs;
 
 		PrivateDependencyModuleNames.AddRange(
-			new string[]
-			{
+			new string[] {
 				"ApplicationCore",
 				"AssetRegistry",
-				"AssetTools",
-				"BlueprintGraph",
 				"Core",
 				"CoreUObject",
-				"EditorStyle",
 				"Engine",
-				"GraphEditor",
 				"InputCore",
 				"Json",
-				"Kismet",
-				"KismetCompiler",
-				"KismetWidgets",
-				"Projects",
-				"PropertyEditor",
 				"Slate",
 				"SlateCore",
-				"ToolMenus",
-				"UnrealEd",
 				"WebAPI",
 			}
 		);
+		
+		if (Target.bBuildEditor == true)
+		{
+			PrivateDependencyModuleNames.AddRange(
+				new string[] {
+					"AssetTools",
+					"BlueprintGraph",
+					"EditorStyle",
+					"GraphEditor",
+					"Kismet",
+					"KismetCompiler",
+					"KismetWidgets",
+					"Projects",
+					"PropertyEditor",
+					"ToolMenus",
+					"UnrealEd"
+				});
+		}
 	}
 }
diff --git a/Engine/Plugins/Web/WebAPI/Source/WebAPIEditor/Private/Details/Widgets/SWebAPISchemaOperationRow.cpp b/Engine/Plugins/Web/WebAPI/Source/WebAPIEditor/Private/Details/Widgets/SWebAPISchemaOperationRow.cpp
--- a/Engine/Plugins/Web/WebAPI/Source/WebAPIEditor/Private/Details/Widgets/SWebAPISchemaOperationRow.cpp
+++ b/Engine/Plugins/Web/WebAPI/Source/WebAPIEditor/Private/Details/Widgets/SWebAPISchemaOperationRow.cpp
@@ -20,9 +20,7 @@ void SWebAPISchemaOperationRow::Construct(const FArguments& InArgs, const TShare
 		.Content()
 		[
 			SNew(SBorder)
-			.BorderBackgroundColor(FLinearColor(0,0,0,0))
-			// @todo: this prevents selection highlight, check slack correspondence
-			//.BorderImage(FAppStyle::GetBrush("ToolPanel.GroupBorder"))
+			.BorderBackgroundColor(FLinearColor(0,0,0,0)) // hide drawn border
 			.Padding(4)
 			[
 				SNew(SBox)
diff --git a/Engine/Plugins/Web/WebAPI/Source/WebAPIEditor/Private/Details/Widgets/SWebAPISchemaServiceRow.cpp b/Engine/Plugins/Web/WebAPI/Source/WebAPIEditor/Private/Details/Widgets/SWebAPISchemaServiceRow.cpp
--- a/Engine/Plugins/Web/WebAPI/Source/WebAPIEditor/Private/Details/Widgets/SWebAPISchemaServiceRow.cpp
+++ b/Engine/Plugins/Web/WebAPI/Source/WebAPIEditor/Private/Details/Widgets/SWebAPISchemaServiceRow.cpp
@@ -16,9 +16,7 @@ void SWebAPISchemaServiceRow::Construct(const FArguments& InArgs, const TSharedR
 		.Content()
 		[
 			SNew(SBorder)
-			.BorderBackgroundColor(FLinearColor(0,0,0,0))
-			// @todo: this prevents selection highlight, check slack correspondence
-			//.BorderImage(FAppStyle::GetBrush("ToolPanel.GroupBorder"))
+			.BorderBackgroundColor(FLinearColor(0,0,0,0)) // hide drawn border
 			.Padding(4)
 			[
 				SNew(SBox)
diff --git a/Engine/Plugins/Web/WebAPI/Source/WebAPIEditor/Private/Factories/WebAPIDefinitionFactory.cpp b/Engine/Plugins/Web/WebAPI/Source/WebAPIEditor/Private/Factories/WebAPIDefinitionFactory.cpp
--- a/Engine/Plugins/Web/WebAPI/Source/WebAPIEditor/Private/Factories/WebAPIDefinitionFactory.cpp
+++ b/Engine/Plugins/Web/WebAPI/Source/WebAPIEditor/Private/Factories/WebAPIDefinitionFactory.cpp
@@ -311,7 +311,6 @@ EReimportResult::Type UWebAPIDefinitionFactory::Reimport(UObject* InObj)
 
 FString UWebAPIDefinitionFactory::ReadFileContents(const FString& InFileName) const
 {
-	// @todo: handle URLs
 	FString FileContents;
 	if (!FFileHelper::LoadFileToString(FileContents, *InFileName))
 	{
diff --git a/Engine/Plugins/Web/WebAPI/Source/WebAPIEditor/Private/WebAPIEditorWidgetUtilities.h b/Engine/Plugins/Web/WebAPI/Source/WebAPIEditor/Private/WebAPIEditorWidgetUtilities.h
--- a/Engine/Plugins/Web/WebAPI/Source/WebAPIEditor/Private/WebAPIEditorWidgetUtilities.h
+++ b/Engine/Plugins/Web/WebAPI/Source/WebAPIEditor/Private/WebAPIEditorWidgetUtilities.h
@@ -4,23 +4,10 @@
 
 #include "CoreMinimal.h"
 #include "Styling/SlateColor.h"
-#include "UObject/Object.h"
-#include "WebAPIEditorWidgetUtilities.generated.h"
-
-/**
- * 
- */
-UCLASS()
-class WEBAPIEDITOR_API UWebAPIEditorWidgetUtilities : public UObject
-{
-	GENERATED_BODY()
-};
-
 
 class FWebAPIEditorWidgetsUtilities
 {
 public:
-	// @todo: verb type should be enum
 	static FSlateColor GetColorForVerb(uint8 InVerb) { return FSlateColor(FColor{255, 128, 64}); }
 	static FName GetIconNameForVerb(uint8 InVerb) { return NAME_None; }
 };
diff --git a/Engine/Plugins/Web/WebAPI/Source/WebAPILiquidJS/Private/WebAPILiquidJSProcess.cpp b/Engine/Plugins/Web/WebAPI/Source/WebAPILiquidJS/Private/WebAPILiquidJSProcess.cpp
--- a/Engine/Plugins/Web/WebAPI/Source/WebAPILiquidJS/Private/WebAPILiquidJSProcess.cpp
+++ b/Engine/Plugins/Web/WebAPI/Source/WebAPILiquidJS/Private/WebAPILiquidJSProcess.cpp
@@ -1,26 +1,24 @@
-// Copyright Epic Games, Inc. All Rights Reserved.
+﻿// Copyright Epic Games, Inc. All Rights Reserved.
 
 #include "WebAPILiquidJSProcess.h"
 
+#include "INetworkingWebSocket.h"
 #include "IWebAPILiquidJSModule.h"
-#include "Interfaces/IPluginManager.h"
-#include "GenericPlatform/GenericPlatformFile.h"
-#include "GenericPlatform/GenericPlatformProcess.h"
-#include "HAL/PlatformFileManager.h"
-#include "Modules/ModuleManager.h"
+#include "IWebSocketNetworkingModule.h"
 #include "SocketSubsystem.h"
 #include "WebAPILiquidJSLog.h"
 #include "WebAPILiquidJSSettings.h"
-#include "INetworkingWebSocket.h"
+#include "Async/Async.h"
+#include "GenericPlatform/GenericPlatformFile.h"
+#include "GenericPlatform/GenericPlatformProcess.h"
+#include "HAL/PlatformFileManager.h"
+#include "Interfaces/IPluginManager.h"
 
 #define LOCTEXT_NAMESPACE "WebAPILiquidJS"
 
 FWebAPILiquidJSProcess::FWebAPILiquidJSProcess()
 	: Status(EStatus::Stopped)
-{
-}
-
-FWebAPILiquidJSProcess::~FWebAPILiquidJSProcess()
+	, bForceBuildWebApp(false)
 {
 }
 
@@ -37,7 +35,7 @@ bool FWebAPILiquidJSProcess::TryStart()
 
 	if (!PlatformFile.DirectoryExists(*Root))
 	{
-		UE_LOG(LogWebAPILiquidJS, Warning, TEXT("WebApp folder doesn't exist (%s)"), *Root);
+		UE_LOG(LogWebAPILiquidJS, Warning, TEXT("WebAPIGeneratorApp folder doesn't exist (%s)"), *Root);
 		return false;
 	}
 
@@ -91,14 +89,6 @@ void FWebAPILiquidJSProcess::SetExternalLoggerEnabled(bool bEnableExternalLog) c
 		bool bIsValidIp = false;
 		Address->SetIp(TEXT("127.0.0.1"), bIsValidIp);
 		Address->SetPort(Settings->Port + 2);
-		//WebSocketLoggerConnection = FModuleManager::LoadModuleChecked<IWebSocketNetworkingModule>(TEXT("WebSocketNetworking")).CreateConnection(*Address);
-	}
-
-	IWebAPILiquidJSModuleInterface* WebAPILiquidJSModule = FModuleManager::GetModulePtr<IWebAPILiquidJSModuleInterface>("WebAPILiquidJS");
-	if (WebAPILiquidJSModule)
-	{
-		// @todo
-		//WebAPILiquidJSModule->SetExternalRemoteWebSocketLoggerConnection(WebSocketLoggerConnection);
 	}
 }
 
@@ -134,7 +124,7 @@ uint32 FWebAPILiquidJSProcess::Run()
 									Settings->WebSocketServerPort,
 									Settings->HttpServerPort);
 
-	if (Settings->bForceWebAppBuildAtStartup)
+	if (Settings->bForceWebAppBuildAtStartup || bForceBuildWebApp)
 	{
 		Args.Append(TEXT("--build "));
 	}
@@ -260,7 +250,7 @@ uint32 FWebAPILiquidJSProcess::Run()
 	{
 		TaskNotification->SetComplete(
 			ErrorTitle,
-			LOCTEXT("WebAPILiquidJS_Exited", "WebApp exited"),
+			LOCTEXT("WebAPILiquidJS_LaunchFailed", "WebApp exited"),
 			false
 		);
 	}
@@ -269,11 +259,33 @@ uint32 FWebAPILiquidJSProcess::Run()
 		UE_LOG(LogWebAPILiquidJS, Log, TEXT("WebApp exited"));
 	}
 
+	EStatus PreStopStatus = Status;
 	Status = EStatus::Stopped;
 
 	FPlatformProcess::ClosePipe(ReadPipe, WritePipe);
 	Process.Reset();
 
+	// If haven't already, try again with --build flag
+	if(PreStopStatus == EStatus::Error && !bForceBuildWebApp && !Settings->bForceWebAppBuildAtStartup)
+	{
+		bForceBuildWebApp = true;
+
+		Async(EAsyncExecution::TaskGraphMainThread, [this]()
+		{
+			FAsyncTaskNotificationConfig NotificationConfig;
+			NotificationConfig.TitleText = LOCTEXT("WebAPILiquidJS_Launch", "Launching WebAPILiquidJS");
+	#if !NO_LOGGING
+			NotificationConfig.LogCategory = &LogWebAPILiquidJS;
+	#endif
+			NotificationConfig.bIsHeadless = false;
+			
+			TaskNotification.Release();
+			TaskNotification = MakeUnique<FAsyncTaskNotification>(NotificationConfig);
+		}).Wait();
+
+		return Run();
+	}
+
 	return 0;
 }
 
diff --git a/Engine/Plugins/Web/WebAPI/Source/WebAPILiquidJS/Private/WebAPILiquidJSProcess.h b/Engine/Plugins/Web/WebAPI/Source/WebAPILiquidJS/Private/WebAPILiquidJSProcess.h
--- a/Engine/Plugins/Web/WebAPI/Source/WebAPILiquidJS/Private/WebAPILiquidJSProcess.h
+++ b/Engine/Plugins/Web/WebAPI/Source/WebAPILiquidJS/Private/WebAPILiquidJSProcess.h
@@ -24,7 +24,7 @@ public:
 public:
 	FWebAPILiquidJSProcess();
 	
-	virtual ~FWebAPILiquidJSProcess();
+	virtual ~FWebAPILiquidJSProcess() override = default;
 
 	/** Starts the webapp external process, first launch will take some time as it needs to compile the webapp */
 	bool TryStart();
@@ -56,4 +56,7 @@ private:
 
 	/** The current status of the web app. */
 	std::atomic<EStatus> Status;
-};
\ No newline at end of file
+
+	/** Can by used override Settings for retry. */
+	std::atomic<bool> bForceBuildWebApp;
+};
diff --git a/Engine/Plugins/Web/WebAPI/Source/WebAPIOpenAPI/Private/Factories/WebAPIOpenAPIFactoryBase.h b/Engine/Plugins/Web/WebAPI/Source/WebAPIOpenAPI/Private/Factories/WebAPIOpenAPIFactoryBase.h
--- a/Engine/Plugins/Web/WebAPI/Source/WebAPIOpenAPI/Private/Factories/WebAPIOpenAPIFactoryBase.h
+++ b/Engine/Plugins/Web/WebAPI/Source/WebAPIOpenAPI/Private/Factories/WebAPIOpenAPIFactoryBase.h
@@ -9,7 +9,7 @@
 
 /** Common base factory between Swagger and OpenAPI factories. */
 UCLASS(Abstract)
-class WEBAPIOPENAPI_API UWebAPIOpenAPIFactoryBase
+class UWebAPIOpenAPIFactoryBase
 	: public UWebAPIDefinitionFactory
 {
 	GENERATED_BODY()
diff --git a/Engine/Plugins/Web/WebAPI/Source/WebAPIOpenAPI/Private/Tests/WebAPIOpenAPI3.spec.cpp b/Engine/Plugins/Web/WebAPI/Source/WebAPIOpenAPI/Private/Tests/WebAPIOpenAPI3.spec.cpp
--- a/Engine/Plugins/Web/WebAPI/Source/WebAPIOpenAPI/Private/Tests/WebAPIOpenAPI3.spec.cpp
+++ b/Engine/Plugins/Web/WebAPI/Source/WebAPIOpenAPI/Private/Tests/WebAPIOpenAPI3.spec.cpp
@@ -3,6 +3,7 @@
 #include "Misc/AutomationTest.h"
 
 #if WITH_EDITOR
+
 #include "WebAPIEditorSettings.h"
 #include "Dom/WebAPIParameter.h"
 #include "Misc/AutomationTest.h"
@@ -10,7 +11,6 @@
 #include "Misc/Paths.h"
 #include "Serialization/JsonSerializer.h"
 #include "V3/WebAPIOpenAPIConverter.h"
-#include "V3/WebAPIOpenAPIFactory.h"
 #include "V3/WebAPIOpenAPISchema.h"
 
 #if WITH_DEV_AUTOMATION_TESTS
@@ -22,7 +22,7 @@ BEGIN_DEFINE_SPEC(FWebAPIOpenAPI3Spec,
 	TSharedPtr<UE::WebAPI::OpenAPI::V3::FOpenAPIObject> InputDefinition;
 	TStrongObjectPtr<UWebAPIDefinition> OutputDefinition;
 	//TSharedPtr<UE::WebAPI::OpenAPI::FWebAPIOpenAPISchemaConverter> Converter;
-	TStrongObjectPtr<UWebAPIDefinitionFactory> Factory;
+	//TStrongObjectPtr<UWebAPIDefinitionFactory> Factory;
 
 	template <class SchemaObjectType>
 	void LoadFromJson(const TSharedPtr<SchemaObjectType>& InSchemaObject, const TSharedPtr<FJsonObject>& InJsonObject)
@@ -38,7 +38,7 @@ BEGIN_DEFINE_SPEC(FWebAPIOpenAPI3Spec,
 	}
 
 	template <class SchemaObjectType>
-	void LoadFromJson(const TArray<TSharedPtr<SchemaObjectType>>& InSchemaObject, const TSharedPtr<FJsonValueArray>& InJsonValue)
+	void LoadFromJson(TArray<TSharedPtr<SchemaObjectType>>& InSchemaObject, const TSharedPtr<FJsonValueArray>& InJsonValue)
 	{
 		const TArray<TSharedPtr<FJsonValue>>* JsonArray = nullptr;
 		if(InJsonValue->TryGetArray(JsonArray))
@@ -46,8 +46,9 @@ BEGIN_DEFINE_SPEC(FWebAPIOpenAPI3Spec,
 			for(const TSharedPtr<FJsonValue>& JsonItem : *JsonArray)
 			{
 				TSharedPtr<SchemaObjectType> SchemaObject = MakeShared<SchemaObjectType>();
-				SchemaObject->FromJson(JsonItem->AsObject().ToSharedRef());				
-			}			
+				SchemaObject->FromJson(JsonItem->AsObject().ToSharedRef());
+				InSchemaObject.Add(SchemaObject);
+			}
 		}
 	}
 
@@ -86,7 +87,7 @@ BEGIN_DEFINE_SPEC(FWebAPIOpenAPI3Spec,
 
 		OutputDefinition = TStrongObjectPtr(NewObject<UWebAPIDefinition>());
 
-		Factory = TStrongObjectPtr(NewObject<UWebAPIOpenAPIFactory>());
+		//Factory = TStrongObjectPtr(NewObject<UWebAPIOpenAPIFactory>());
 
 		return MakeShared<UE::WebAPI::OpenAPI::FWebAPIOpenAPISchemaConverter>(
 			InputDefinition,
@@ -115,7 +116,7 @@ BEGIN_DEFINE_SPEC(FWebAPIOpenAPI3Spec,
 		
 		OutputDefinition = TStrongObjectPtr(NewObject<UWebAPIDefinition>());
 
-		Factory = TStrongObjectPtr(NewObject<UWebAPIOpenAPIFactory>());
+		//Factory = TStrongObjectPtr(NewObject<UWebAPIOpenAPIFactory>());
 
 		return MakeShared<UE::WebAPI::OpenAPI::FWebAPIOpenAPISchemaConverter>(
 			InputDefinition,
@@ -126,7 +127,8 @@ BEGIN_DEFINE_SPEC(FWebAPIOpenAPI3Spec,
 
 	FString GetSampleFile(const FString& InName) const
 	{
-		FString FilePath = FPaths::Combine(FPaths::ProjectPluginsDir(),
+		FString FilePath = FPaths::Combine(FPaths::EnginePluginsDir(),
+			TEXT("Web"),
 			TEXT("WebAPI"), TEXT("Source"), TEXT("WebAPIOpenAPI"),
 			TEXT("Private"), TEXT("Tests"), TEXT("Samples"), TEXT("V3"), InName + TEXT(".json"));
 		ensure(FPaths::FileExists(FilePath));
@@ -155,7 +157,7 @@ BEGIN_DEFINE_SPEC(FWebAPIOpenAPI3Spec,
 	
 		TArray<TSharedPtr<FJsonValue>> JsonValueArray;
 		FJsonSerializer::Deserialize(TJsonReaderFactory<TCHAR>::Create(FileContents), JsonValueArray);
-		return MakeShared<FJsonValueArray>(JsonValueArray);
+		return  MakeShared<FJsonValueArray>(JsonValueArray);
 	}
 
 	template <class InputModelType, class OutputModelType>
@@ -276,12 +278,11 @@ BEGIN_DEFINE_SPEC(FWebAPIOpenAPI3Spec,
 
 END_DEFINE_SPEC(FWebAPIOpenAPI3Spec)
 
-// @todo: remove when done: -ExecCmds="Automation RunTests Plugin.WebAPI.OpenAPI3" -testexit="Automation Test Queue Empty" -unattended  -nopause
 void FWebAPIOpenAPI3Spec::Define()
 {
 	BeforeEach([this]
 	{
-		Factory.Reset();
+		//Factory.Reset();
 		InputDefinition.Reset();
 		OutputDefinition.Reset();
 	});
@@ -345,38 +346,6 @@ void FWebAPIOpenAPI3Spec::Define()
 						TestFoundProperty("Name property found", TagModel, TEXT("Name"), TEXT("String"));
 					}
 				}
-
-				// Test Parameters
-				{
-					TArray<TObjectPtr<UWebAPIParameter>> ConvertedParameters;
-					for(const TObjectPtr<UWebAPIModelBase>& Model : AllConvertedModels)
-					{
-						if(Model->IsA(UWebAPIParameter::StaticClass()))
-						{
-							ConvertedParameters.Add(Cast<UWebAPIParameter>(Model));
-						}
-					}
-					
-					TestEqual("Converted 2 parameters", ConvertedParameters.Num(), 2);
-
-					const TObjectPtr<UWebAPIParameter> SkipParamParameter = FindNamedModel<UWebAPIParameter>(ConvertedParameters, TEXT("SkipParam"));
-					if(TestNotNull("SkipParam parameter found", SkipParamParameter))
-					{
-						TestEqual("Parameter name is skip", SkipParamParameter->Name.ToString(true), TEXT("skip"));
-						TestEqual("Parameter stored in query", SkipParamParameter->Storage, EWebAPIParameterStorage::Query);
-						TestTrue("Parameter is required", SkipParamParameter->bIsRequired);
-						TestEqual("Parameter is 32 bit integer", SkipParamParameter->Property->Type.ToString(true), TEXT("int32"));
-					}
-
-					const TObjectPtr<UWebAPIParameter> LimitParamParameter = FindNamedModel<UWebAPIParameter>(ConvertedParameters, TEXT("LimitParam"));
-					if(TestNotNull("LimitParam parameter found", LimitParamParameter))
-					{
-						TestEqual("Parameter name is limit", LimitParamParameter->Name.ToString(true), TEXT("limit"));
-						TestEqual("Parameter stored in query", LimitParamParameter->Storage, EWebAPIParameterStorage::Query);
-						TestTrue("Parameter is required", LimitParamParameter->bIsRequired);
-						TestEqual("Parameter is 32 bit integer", LimitParamParameter->Property->Type.ToString(true), TEXT("int32"));
-					}
-				}
 			}
 		});
 	});
@@ -429,7 +398,6 @@ void FWebAPIOpenAPI3Spec::Define()
 
 	Describe("ServersObject", [this]
 	{
-		// @todo: support server variables?
 		It("Converts with Variables", [this]
 		{
 			const FString FilePath = GetSampleFile("ServerObjectExample_Variables");
@@ -443,7 +411,7 @@ void FWebAPIOpenAPI3Spec::Define()
 			const bool bWasConverted = Converter->Convert();
 			if(TestTrue("ServersObject converted to WebAPI", bWasConverted))
 			{
-				const FString ExpectedUrl = TEXT("https://{username}.gigantic-server.com:{port}/{basePath}");
+				const FString ExpectedUrl = TEXT("{username}.gigantic-server.com:{port}");
 				TestEqual(FString::Printf(TEXT("Host is \"%s\""), *ExpectedUrl), OutputDefinition->GetWebAPISchema()->Host, ExpectedUrl);
 			}
 		});
diff --git a/Engine/Plugins/Web/WebAPI/Source/WebAPIOpenAPI/Private/Tests/WebAPISwagger.spec.cpp b/Engine/Plugins/Web/WebAPI/Source/WebAPIOpenAPI/Private/Tests/WebAPISwagger.spec.cpp
--- a/Engine/Plugins/Web/WebAPI/Source/WebAPIOpenAPI/Private/Tests/WebAPISwagger.spec.cpp
+++ b/Engine/Plugins/Web/WebAPI/Source/WebAPIOpenAPI/Private/Tests/WebAPISwagger.spec.cpp
@@ -1,6 +1,7 @@
 ﻿// Copyright Epic Games, Inc. All Rights Reserved.
 
 #if WITH_EDITOR
+
 #include "WebAPIEditorSettings.h"
 #include "Dom/WebAPIParameter.h"
 #include "Misc/AutomationTest.h"
@@ -9,25 +10,22 @@
 #include "Serialization/JsonSerializer.h"
 #include "V2/WebAPISwaggerConverter.h"
 #include "V2/WebAPISwaggerConverter.inl"
-#include "V2/WebAPISwaggerFactory.h"
 #include "V2/WebAPISwaggerSchema.h"
 
 #if WITH_DEV_AUTOMATION_TESTS
 
-using namespace UE::WebAPI::OpenAPI;
-
 BEGIN_DEFINE_SPEC(FWebAPISwaggerSpec,
-				TEXT("Plugin.WebAPI.Swagger"),
-				EAutomationTestFlags::EditorContext | EAutomationTestFlags::EngineFilter | EAutomationTestFlags::ApplicationContextMask)
+	TEXT("Plugin.WebAPI.Swagger"),
+	EAutomationTestFlags::EditorContext | EAutomationTestFlags::EngineFilter | EAutomationTestFlags::ApplicationContextMask)
 
 	TSharedPtr<UE::WebAPI::OpenAPI::V2::FSwagger> InputDefinition;
 	TStrongObjectPtr<UWebAPIDefinition> OutputDefinition;
 	TSharedPtr<UE::WebAPI::Swagger::FWebAPISwaggerSchemaConverter> Converter;
-	TStrongObjectPtr<UWebAPIDefinitionFactory> Factory;
 
 	FString GetSampleFile(const FString& InName) const
 	{
-		FString FilePath = FPaths::Combine(FPaths::ProjectPluginsDir(),
+		FString FilePath = FPaths::Combine(FPaths::EnginePluginsDir(),
+			TEXT("Web"),
 			TEXT("WebAPI"), TEXT("Source"), TEXT("WebAPIOpenAPI"),
 			TEXT("Private"), TEXT("Tests"), TEXT("Samples"), TEXT("V2"), InName + TEXT(".json"));
 		ensure(FPaths::FileExists(FilePath));
@@ -36,17 +34,7 @@ BEGIN_DEFINE_SPEC(FWebAPISwaggerSpec,
 
 	FString GetAPISample() const
 	{
-		return GetSampleFile(TEXT("petstore_V2.json"));
-	}
-
-	FString GetDigitalTwinsManagementAPISample() const
-	{
-		return GetSampleFile(TEXT("digitaltwins"));
-	}
-
-	FString GetDigitalTwinsAPISample() const
-	{
-		return GetSampleFile(TEXT("digitaltwins_data"));
+		return GetSampleFile(TEXT("petstore_v2"));
 	}
 
 	TSharedPtr<FJsonObject> LoadJson(const FString& InFile) const
@@ -61,12 +49,13 @@ BEGIN_DEFINE_SPEC(FWebAPISwaggerSpec,
 
 END_DEFINE_SPEC(FWebAPISwaggerSpec)
 
-// @todo: remove when done: -ExecCmds="Automation RunTests Plugin.WebAPI.Parser.Swagger2" -testexit="Automation Test Queue Empty" -unattended  -nopause
 void FWebAPISwaggerSpec::Define()
 {
+	using namespace UE::WebAPI::OpenAPI;
+
 	BeforeEach([this]
 	{
-		const FString FilePath = GetDigitalTwinsManagementAPISample();
+		const FString FilePath = GetAPISample();
 		const TSharedPtr<FJsonObject> JsonObject = LoadJson(FilePath);
 
 		InputDefinition = MakeShared<UE::WebAPI::OpenAPI::V2::FSwagger>();
@@ -79,8 +68,6 @@ void FWebAPISwaggerSpec::Define()
 			OutputDefinition->GetWebAPISchema(),
 			OutputDefinition->GetMessageLog().ToSharedRef(),
 			OutputDefinition->GetProviderSettings());
-
-		Factory = TStrongObjectPtr(NewObject<UWebAPISwaggerFactory>());
 	});
 	
 	Describe("PetStore", [this]
@@ -134,68 +121,11 @@ void FWebAPISwaggerSpec::Define()
 		});
 	});
 
-	Describe("DigitalTwinsManagement", [this]
-	{
-		It("Parses", [this]
-		{
-			const FString FilePath = GetDigitalTwinsManagementAPISample();
-			const TSharedPtr<FJsonObject> JsonObject = LoadJson(FilePath);
-
-			UE::WebAPI::OpenAPI::V2::FSwagger Object;
-			Object.FromJson(JsonObject.ToSharedRef());
-
-			TestEqual("Swagger Version", Object.Swagger, "2.0");
-			TestEqual("Info Version", Object.Info->Version, "2020-12-01");
-			if(TestTrue("Host Set", Object.Host.IsSet()))
-			{
-				TestEqual("Host", Object.Host.GetValue(), "management.azure.com");
-			}
-
-			if(TestTrue("Schemes Set", Object.Schemes.IsSet()))
-			{
-				TestEqual("Scheme Num", Object.Schemes->Num(), 1);
-			}
-
-			TestEqual("Paths Num", Object.Paths.Num(), 11);
-				
-			if(TestTrue("Security Definitions Set", Object.SecurityDefinitions.IsSet()))
-			{
-				TestEqual("Security Definitions Num", Object.SecurityDefinitions->Num(), 1);
-			}
-				
-			if(TestTrue("Definitions Set", Object.Definitions.IsSet()))
-			{
-				TestEqual("Definitions Num", Object.Definitions->Num(), 30);
-			}
-		});
-	});
-
-	Describe("Model Composition", [this]
-	{
-		// @todo: meaningful name
-		It("All-Of", [this]
-		{
-			const TSharedPtr<UE::WebAPI::OpenAPI::V2::FSchema>* InputModel = InputDefinition->Definitions->Find(TEXT("ServiceBus"));
-
-			const TObjectPtr<UWebAPIModel> OutputModel = Converter->ConvertModel(*InputModel);
-			if(TestFalse("Converted Model", OutputModel.IsNull()))
-			{
-				if(TestTrue("All-Of is set", (*InputModel)->AllOf.IsSet()))
-				{
-					for(UE::Json::TJsonReference<UE::WebAPI::OpenAPI::V2::FSchema>& SubType : (*InputModel)->AllOf.GetValue())
-					{
-					
-					}
-				}
-			}
-		});
-	});
-
 	Describe("Type Conversion", [this]
 	{
 		BeforeEach([this]
 		{
-			const FString FilePath = GetDigitalTwinsAPISample();
+			const FString FilePath = GetAPISample();
 			const TSharedPtr<FJsonObject> JsonObject = LoadJson(FilePath);
 
 			InputDefinition = MakeShared<UE::WebAPI::OpenAPI::V2::FSwagger>();
@@ -208,166 +138,12 @@ void FWebAPISwaggerSpec::Define()
 				OutputDefinition->GetWebAPISchema(),
 				OutputDefinition->GetMessageLog().ToSharedRef(),
 				OutputDefinition->GetProviderSettings());
-
-			Factory = TStrongObjectPtr(NewObject<UWebAPISwaggerFactory>());
-		});
-		// // @todo: meaningful name
-		// It("Test", [this]()
-		// {
-		// 	const TSharedPtr<UE::WebAPI::OpenAPI::V2::FSchema>* DigitalTwinsProperties = InputDefinition->Definitions->Find(TEXT("DigitalTwinsProperties"));
-		//
-		// 	const TObjectPtr<UWebAPITypeInfo> Type = Converter->ResolveType(*DigitalTwinsProperties, TEXT("DigitalTwinsProperties"));
-		//
-		// 	TestTrue("TypeInfo is not null", Type != nullptr);
-		// });
-
-		It("Enum that's non-optional and has a single value should default to that value", [this]
-		{
-			check(Converter->Convert());
-			Factory->PostImportWebAPI(OutputDefinition.Get()).Wait();
-			
-			const FString EnumNameToTest = TEXT("ApiVersion");
-
-			TObjectPtr<UWebAPIModelBase>* FoundModel = OutputDefinition->GetWebAPISchema()->Models.FindByPredicate([&](const TObjectPtr<UWebAPIModelBase>& InModel)
-			{
-				return InModel->IsA(UWebAPIEnum::StaticClass()) && InModel->GetName().Contains(EnumNameToTest);
-			});
-
-			if(!TestNotNull("Found Enum", FoundModel)
-				|| !TestTrue("Found Enum not nullptr", !FoundModel->IsNull()))
-			{
-				return false;				
-			}
-
-			const TObjectPtr<UWebAPIEnum> FoundEnum = Cast<UWebAPIEnum>(*FoundModel);
-			TestEqual("Single value", FoundEnum->Values.Num(), 1);
-			TestTrue("Default value is the first value", FoundEnum->GetDefaultValue().Contains(FoundEnum->Values[0]->Name.ToMemberName(TEXT("EV"))));
-
- 			return true;
 		});
 
-		It("String Array that's optional should have an empty array as a default value", [this]
-		{
-			check(Converter->Convert());
-			Factory->PostImportWebAPI(OutputDefinition.Get()).Wait();
-			
-			const FString ObjectNameToTest = TEXT("IncomingRelationshipCollection");
-			const FString PropertyNameToTest = TEXT("Value");
-
-			TObjectPtr<UWebAPIModelBase>* FoundModelBase = OutputDefinition->GetWebAPISchema()->Models.FindByPredicate([&](const TObjectPtr<UWebAPIModelBase>& InModel)
-			{
-				return InModel->IsA(UWebAPIModel::StaticClass()) && InModel->GetName().Contains(ObjectNameToTest);
-			});
-
-			if(!TestNotNull("Found Model", FoundModelBase)
-				|| !TestTrue("Found Model not nullptr", !FoundModelBase->IsNull()))
-			{
-				return false;				
-			}
-
-			const TObjectPtr<UWebAPIModel> FoundModel = Cast<UWebAPIModel>(*FoundModelBase);
-
-			const TObjectPtr<UWebAPIProperty>* StringArrayProperty = FoundModel->Properties.FindByPredicate([&](const TObjectPtr<UWebAPIProperty>& InProperty)
-			{
-				return InProperty->Name.ToString(true) == PropertyNameToTest;				
-			});
-
-			const FString DefaultValue = (*StringArrayProperty)->GetDefaultValue();
-			TestTrue("Empty Array default value", DefaultValue.Len() <= 3);
-
-			return true;
-		});
-		
 		It("Default value should be set in metadata", [this]
 		{
 
 		});
-
-		It("Properties marked as required in the root of a definition are required in the output object(s)", [this]
-		{
-			check(Converter->Convert());
-			Factory->PostImportWebAPI(OutputDefinition.Get()).Wait();
-					
-			const FString EnumNameToTest = TEXT("ApiVersion");
-
-			TObjectPtr<UWebAPIModelBase>* FoundModel = OutputDefinition->GetWebAPISchema()->Models.FindByPredicate([&](const TObjectPtr<UWebAPIModelBase>& InModel)
-			{
-				return InModel->GetName().Contains(EnumNameToTest);
-			});
-
-			if(!TestNotNull("Found Enum", FoundModel)
-				|| TestTrue("Found Enum not nullptr", !FoundModel->IsNull()))
-			{
-				return false;				
-			}
-
-			const TObjectPtr<UWebAPIEnum> FoundEnum = Cast<UWebAPIEnum>(*FoundModel);
-			const TObjectPtr<UWebAPIEnumValue>* FoundEnumValue = FoundEnum->Values.FindByPredicate([&](const TObjectPtr<UWebAPIEnumValue>& InEnumValue)
-			{
-				return InEnumValue->Name == OutputDefinition->GetProviderSettings().GetUnsetEnumValueName();
-			});
-
-			TestNull("Unset value not found", FoundEnumValue);
-
-			return true;
-		});
-	});
-
-	Describe("Parameter", [this]
-	{
-		It("Resolves schema-less, built-in type", [this]
-		{
-			const TSharedPtr<UE::WebAPI::OpenAPI::V2::FParameter>* InputParameter = InputDefinition->Parameters->Find(TEXT("resourceName"));
-
-			const TObjectPtr<UWebAPIParameter> OutputParameter = Converter->ConvertParameter(*InputParameter);
-			if(TestFalse("Converted Parameter", OutputParameter.IsNull()))
-			{
-				TestEqual("Type is String", OutputParameter->Type.ToString(true), TEXT("String"));
-				TestTrue("Object is named correctly", OutputParameter->GetFName().ToString().StartsWith(TEXT("WebAPIParameter_ResourceName")));
-			}
-		});
-
-		It("Resolves schema-less, generated enum type", [this]
-		{
-			const TSharedPtr<UE::WebAPI::OpenAPI::V2::FParameter>* InputParameter = InputDefinition->Parameters->Find(TEXT("api-version"));
-
-			const TObjectPtr<UWebAPIParameter> OutputParameter = Converter->ConvertParameter(*InputParameter);
-			if(TestFalse("Converted Parameter", OutputParameter.IsNull()))
-			{
-				TestTrue("Type is Enum", OutputParameter->Type.TypeInfo->IsEnum());
-
-				UWebAPIEnum* Model = Cast<UWebAPIEnum>(OutputParameter->Type.TypeInfo->Model.Get());
-				if(TestTrue("Underlying Enum", Model != nullptr))
-				{
-					// TestTrue("Type is Enum", OutputParameter->Type.TypeInfo->IsEnum());
-				}
-			}
-		});
-	});
-		
-	Describe("Operation", [this]
-	{
-		It("Resolves referenced parameters", [this]
-		{
-			const FString Url = TEXT("/subscriptions/{subscriptionId}/resourceGroups/{resourceGroupName}/providers/Microsoft.DigitalTwins/digitalTwinsInstances/{resourceName}/privateEndpointConnections/{privateEndpointConnectionName}");
-
-			const TSharedPtr<V2::FPath, ESPMode::ThreadSafe>* InputService = InputDefinition->Paths.Find(Url);
-			const FString ServiceTag = TEXT("PrivateEndpoints");
-			const TObjectPtr<UWebAPIService> OutputService = Converter->ConvertService(ServiceTag);
-			
-			const TSharedPtr<UE::WebAPI::OpenAPI::V2::FOperation>& InputOperation = InputDefinition->Paths
-				.Find(Url)
-				->Get()->Get;
-
-			const TObjectPtr<UWebAPIOperation> OutputOperation = Converter->ConvertOperation(Url, TEXT("Get"), InputOperation);
-			const FString OperationName = OutputOperation->Name.ToString(true);
-			
-			if(TestFalse("Converted Operation", OutputOperation.IsNull()))
-			{
-				TestTrue("Object is named correctly", OutputOperation->GetFName().ToString().StartsWith(TEXT("WebAPIOperation_PrivateEndpoints_") + OperationName));
-			}
-		});
-
 	});
 }
 
diff --git a/Engine/Plugins/Web/WebAPI/Source/WebAPIOpenAPI/Private/V2/WebAPISwaggerConverter.h b/Engine/Plugins/Web/WebAPI/Source/WebAPIOpenAPI/Private/V2/WebAPISwaggerConverter.h
--- a/Engine/Plugins/Web/WebAPI/Source/WebAPIOpenAPI/Private/V2/WebAPISwaggerConverter.h
+++ b/Engine/Plugins/Web/WebAPI/Source/WebAPIOpenAPI/Private/V2/WebAPISwaggerConverter.h
@@ -124,3 +124,5 @@ namespace UE::WebAPI::Swagger
 	template <>
 	bool FWebAPISwaggerSchemaConverter::PatchModel(const TSharedPtr<OpenAPI::V2::FParameter>& InSrcSchema, const FWebAPITypeNameVariant& InModelTypeName, const TObjectPtr<UWebAPIModel>& OutModel);
 }
+
+#include "WebAPISwaggerConverter.inl"
diff --git a/Engine/Plugins/Web/WebAPI/Source/WebAPIOpenAPI/Private/V2/WebAPISwaggerConverter.inl b/Engine/Plugins/Web/WebAPI/Source/WebAPIOpenAPI/Private/V2/WebAPISwaggerConverter.inl
--- a/Engine/Plugins/Web/WebAPI/Source/WebAPIOpenAPI/Private/V2/WebAPISwaggerConverter.inl
+++ b/Engine/Plugins/Web/WebAPI/Source/WebAPIOpenAPI/Private/V2/WebAPISwaggerConverter.inl
@@ -2,8 +2,6 @@
 
 #pragma once
 
-#include "WebAPISwaggerConverter.h"
-
 #include "IWebAPIEditorModule.h"
 #include "V2/WebAPISwaggerProvider.h"
 
diff --git a/Engine/Plugins/Web/WebAPI/Source/WebAPIOpenAPI/Private/V2/WebAPISwaggerFactory.cpp b/Engine/Plugins/Web/WebAPI/Source/WebAPIOpenAPI/Private/V2/WebAPISwaggerFactory.cpp
--- a/Engine/Plugins/Web/WebAPI/Source/WebAPIOpenAPI/Private/V2/WebAPISwaggerFactory.cpp
+++ b/Engine/Plugins/Web/WebAPI/Source/WebAPIOpenAPI/Private/V2/WebAPISwaggerFactory.cpp
@@ -12,9 +12,9 @@
 
 struct FWebAPIOperationResponseKeyFuncs : public TDefaultMapKeyFuncs<uint32, TObjectPtr<UWebAPIOperationResponse>, false>
 {
-	static FORCEINLINE uint32 GetSetKey(TObjectPtr<UWebAPIOperationResponse> const& Element) { return Element->Code; }
-	static FORCEINLINE uint32 GetKeyHash(uint32 const& Key) { return Key; }
-	static FORCEINLINE bool Matches(uint32 const& A, uint32 const& B) { return (A == B); }
+	static FORCEINLINE uint32 GetSetKey(const TObjectPtr<UWebAPIOperationResponse>& Element) { return Element->Code; }
+	static FORCEINLINE uint32 GetKeyHash(const uint32& Key) { return Key; }
+	static FORCEINLINE bool Matches(const uint32& A, const uint32& B) { return (A == B); }
 };
 
 UWebAPISwaggerFactory::UWebAPISwaggerFactory()
@@ -47,15 +47,12 @@ TFuture<bool> UWebAPISwaggerFactory::ImportWebAPI(UWebAPIDefinition* InDefinitio
 	ImportData->FileContents = InFileContents;
 
 	// Parse spec
-	// @todo: move call to ConvertToWebAPISchema outside of Factory implementation (to provide Pre and Post extensibility)
 	return Provider->ConvertToWebAPISchema(InDefinition)
 	.Next([](EWebAPIConversionResult bInConversionResult)
 	{
-		// @todo: UI notification
-		//InDefinition->SetWebAPISchema(InWebAPISchema);
-
 		return bInConversionResult == EWebAPIConversionResult::Succeeded;
 	});
 }
 
 #undef LOCTEXT_NAMESPACE
+
diff --git a/Engine/Plugins/Web/WebAPI/Source/WebAPIOpenAPI/Private/V3/WebAPIOpenAPIConverter.h b/Engine/Plugins/Web/WebAPI/Source/WebAPIOpenAPI/Private/V3/WebAPIOpenAPIConverter.h
--- a/Engine/Plugins/Web/WebAPI/Source/WebAPIOpenAPI/Private/V3/WebAPIOpenAPIConverter.h
+++ b/Engine/Plugins/Web/WebAPI/Source/WebAPIOpenAPI/Private/V3/WebAPIOpenAPIConverter.h
@@ -53,8 +53,6 @@ namespace UE::WebAPI::OpenAPI
 		template <typename ObjectType>
 		TSharedPtr<ObjectType> ResolveReference(const UE::Json::TJsonReference<ObjectType>& InJsonReference, bool bInCheck = true);
 
-		//FWebAPINameVariant ResolveModelName(const TObjectPtr<UWebAPIProperty>& InProperty, const TOptional<bool>& bInIsArray, const FWebAPITypeNameVariant& InPossibleName = {});
-
 		FWebAPINameVariant ResolvePropertyName(const TObjectPtr<UWebAPIProperty>& InProperty, const FWebAPITypeNameVariant& InPotentialName, const TOptional<bool>& bInIsArray = {});
 
 		template <typename SchemaType>
@@ -76,10 +74,7 @@ namespace UE::WebAPI::OpenAPI
 		/** Creates a new Property from the supplied source object. */
 		template <typename SchemaType>
 		TObjectPtr<UWebAPIProperty> ConvertProperty(const TSharedPtr<SchemaType>& InSrcSchema, const TObjectPtr<UWebAPIModel>& InModel, const FWebAPINameVariant& InPropertyName = {}, const FString& InDefinitionName = {});
-/*
-		template <>
-		bool ConvertProperty(const FWebAPITypeNameVariant& InModelName,	const FWebAPINameVariant& InPropertyName, const TSharedPtr<OpenAPI::V2::FParameter>& InSchema, const FString& InDefinitionName, const TObjectPtr<UWebAPIProperty>& OutProperty);
-*/
+
 		/** Modifies an existing Model using the supplied source object. */
 		template <typename SchemaType>
 		bool ConvertModel(const TSharedPtr<SchemaType>& InSrcSchema, const FWebAPITypeNameVariant& InModelTypeName, const TObjectPtr<UWebAPIModel>& OutModel);
diff --git a/Engine/Plugins/Web/WebAPI/Source/WebAPIOpenAPI/Private/V3/WebAPIOpenAPISchema.cpp b/Engine/Plugins/Web/WebAPI/Source/WebAPIOpenAPI/Private/V3/WebAPIOpenAPISchema.cpp
--- a/Engine/Plugins/Web/WebAPI/Source/WebAPIOpenAPI/Private/V3/WebAPIOpenAPISchema.cpp
+++ b/Engine/Plugins/Web/WebAPI/Source/WebAPIOpenAPI/Private/V3/WebAPIOpenAPISchema.cpp
@@ -93,7 +93,6 @@ namespace UE::WebAPI::OpenAPI::V3
 		return bSetAField;
 	}
 
-	// @todo: parse header object
 	bool FHeaderObject::FromJson(const TSharedRef<FJsonObject>& InJson)
 	{
 		const bool bSetAField = Super::FromJson(InJson);
