WebAPI validity checks

diff --git a/Engine/Plugins/Web/WebAPI/Source/WebAPIBlueprintGraph/Private/K2Node_WebAPIOperation.cpp b/Engine/Plugins/Web/WebAPI/Source/WebAPIBlueprintGraph/Private/K2Node_WebAPIOperation.cpp
--- a/Engine/Plugins/Web/WebAPI/Source/WebAPIBlueprintGraph/Private/K2Node_WebAPIOperation.cpp
+++ b/Engine/Plugins/Web/WebAPI/Source/WebAPIBlueprintGraph/Private/K2Node_WebAPIOperation.cpp
@@ -100,7 +100,7 @@ namespace UE::WebAPI::Private
 			return false;
 		}
 
-		UK2Node_CustomEvent* CurrentCENode = CompilerContext.SpawnIntermediateEventNode<UK2Node_CustomEvent>(CurrentNode, PinForCurrentDelegateProperty, SourceGraph);
+		UK2Node_CustomEvent* CurrentCustomEventNode = CompilerContext.SpawnIntermediateEventNode<UK2Node_CustomEvent>(CurrentNode, PinForCurrentDelegateProperty, SourceGraph);
 		{
 		UK2Node_AddDelegate* AddDelegateNode = CompilerContext.SpawnIntermediateNode<UK2Node_AddDelegate>(CurrentNode, SourceGraph);
 		AddDelegateNode->SetFromProperty(CurrentProperty, false, CurrentProperty->GetOwnerClass());
@@ -108,17 +108,17 @@ namespace UE::WebAPI::Private
 			bIsErrorFree &= Schema->TryCreateConnection(AddDelegateNode->FindPinChecked(UEdGraphSchema_K2::PN_Self), ProxyObjectPin);
 			bIsErrorFree &= Schema->TryCreateConnection(InOutLastThenPin, AddDelegateNode->FindPinChecked(UEdGraphSchema_K2::PN_Execute));
 			InOutLastThenPin = AddDelegateNode->FindPinChecked(UEdGraphSchema_K2::PN_Then);
-			CurrentCENode->CustomFunctionName = *FString::Printf(TEXT("%s_%s"), *CurrentProperty->GetName(), *CompilerContext.GetGuid(CurrentNode));
-			CurrentCENode->AllocateDefaultPins();
+			CurrentCustomEventNode->CustomFunctionName = *FString::Printf(TEXT("%s_%s"), *CurrentProperty->GetName(), *CompilerContext.GetGuid(CurrentNode));
+			CurrentCustomEventNode->AllocateDefaultPins();
 
-			bIsErrorFree &= CreateDelegateForNewFunction(AddDelegateNode->GetDelegatePin(), CurrentCENode->GetFunctionName(), CurrentNode, SourceGraph, CompilerContext);
-			bIsErrorFree &= CopyEventSignature(CurrentCENode, AddDelegateNode->GetDelegateSignature(), Schema);
+			bIsErrorFree &= CreateDelegateForNewFunction(AddDelegateNode->GetDelegatePin(), CurrentCustomEventNode->GetFunctionName(), CurrentNode, SourceGraph, CompilerContext);
+			bIsErrorFree &= CopyEventSignature(CurrentCustomEventNode, AddDelegateNode->GetDelegateSignature(), Schema);
 		}
 
-		OutLastActivatedThenPin = CurrentCENode->FindPinChecked(UEdGraphSchema_K2::PN_Then);
+		OutLastActivatedThenPin = CurrentCustomEventNode->FindPinChecked(UEdGraphSchema_K2::PN_Then);
 		for (const FOutputPinAndLocalVariable& OutputPair : VariableOutputs) // CREATE CHAIN OF ASSIGMENTS
 		{
-			UEdGraphPin* PinWithData = CurrentCENode->FindPin(OutputPair.OutputPin->PinName);
+			UEdGraphPin* PinWithData = CurrentCustomEventNode->FindPin(OutputPair.OutputPin->PinName);
 			if (PinWithData == nullptr)
 			{
 				continue;
@@ -142,6 +142,11 @@ namespace UE::WebAPI::Private
 
 void UK2Node_WebAPIOperation::AllocateDefaultPins()
 {
+	if(!IsValid())
+	{
+		return;
+	}
+	
 	InvalidatePinTooltips();
 
 	const UEdGraphSchema_K2* K2Schema = GetDefault<UEdGraphSchema_K2>();
@@ -157,8 +162,6 @@ void UK2Node_WebAPIOperation::AllocateDefaultPins()
 		CreatePin(EGPD_Output, UEdGraphSchema_K2::PC_Exec, UEdGraphSchema_K2::PN_Then);
 	}
 
-	const UFunction* Function = GetFactoryFunction();
-
 	TArray<FMulticastDelegateProperty*, TFixedAllocator<2>> DelegateProperties = {
 		UE::WebAPI::Operation::GetPositiveOutcomeDelegate(OperationClass),
 		UE::WebAPI::Operation::GetNegativeOutcomeDelegate(OperationClass)
@@ -210,7 +213,7 @@ void UK2Node_WebAPIOperation::AllocateDefaultPins()
 	}
 
 	bool bAllPinsGood = true;
-	if (Function)
+	if (UFunction* Function = GetFactoryFunction())
 	{
 		TSet<FName> PinsToHide;
 		FBlueprintEditorUtils::GetHiddenPinsForFunction(GetGraph(), Function, PinsToHide);
@@ -710,6 +713,21 @@ bool UK2Node_WebAPIOperation::IsValid() const
 		&& ((bHasCompilerMessage && ErrorType >= EMessageSeverity::Info) || !bHasCompilerMessage);
 }
 
+#if WITH_EDITOR
+EDataValidationResult UK2Node_WebAPIOperation::IsDataValid(TArray<FText>& ValidationErrors)
+{
+	EDataValidationResult ValidationResult = CombineDataValidationResults(Super::IsDataValid(ValidationErrors), EDataValidationResult::Valid);
+	
+	if(OperationClass == nullptr)
+	{
+		ValidationErrors.Add(LOCTEXT("Missing_OperationClass", "OperationClass is invalid or missing"));
+		ValidationResult = EDataValidationResult::Invalid;
+	}
+
+	return ValidationResult;
+}
+#endif
+
 const UK2Node_CustomEvent* UK2Node_WebAPIOperation::GetCustomEventForOutcomeDelegate(
 	const UK2Node_WebAPIOperation* InNode,
 	const FName& InOutcomeName)
diff --git a/Engine/Plugins/Web/WebAPI/Source/WebAPIBlueprintGraph/Private/K2Node_WebAPIOperation.h b/Engine/Plugins/Web/WebAPI/Source/WebAPIBlueprintGraph/Private/K2Node_WebAPIOperation.h
--- a/Engine/Plugins/Web/WebAPI/Source/WebAPIBlueprintGraph/Private/K2Node_WebAPIOperation.h
+++ b/Engine/Plugins/Web/WebAPI/Source/WebAPIBlueprintGraph/Private/K2Node_WebAPIOperation.h
@@ -69,6 +69,11 @@ public:
 	/** Validity of this node. */
 	bool IsValid() const;
 
+#if WITH_EDITOR
+	/** Add node validity to data validation messages. */
+	virtual EDataValidationResult IsDataValid(TArray<FText>& ValidationErrors) override;
+#endif
+
 protected:
 	// Returns the factory function (checked)
 	virtual UFunction* GetFactoryFunction() const;
diff --git a/Engine/Plugins/Web/WebAPI/Source/WebAPIEditor/Public/CodeGen/WebAPICodeGeneratorSettings.h b/Engine/Plugins/Web/WebAPI/Source/WebAPIEditor/Public/CodeGen/WebAPICodeGeneratorSettings.h
--- a/Engine/Plugins/Web/WebAPI/Source/WebAPIEditor/Public/CodeGen/WebAPICodeGeneratorSettings.h
+++ b/Engine/Plugins/Web/WebAPI/Source/WebAPIEditor/Public/CodeGen/WebAPICodeGeneratorSettings.h
@@ -47,7 +47,7 @@ struct WEBAPIEDITOR_API FWebAPICodeGeneratorSettings
 	bool bOverrideGeneratorClass = false;
 
 	/** The Code Generator to use. */
-	UPROPERTY(EditAnywhere, NoClear, Category = "Generator", meta = (MustImplement = "WebAPICodeGeneratorInterface", EditCondition="bOverrideGeneratorClass"))
+	UPROPERTY(EditAnywhere, NoClear, Category = "Generator", meta = (MustImplement = "/Script/WebAPIEditor.WebAPICodeGeneratorInterface", EditCondition="bOverrideGeneratorClass"))
 	TSoftClassPtr<UObject> CodeGeneratorClass;
 
 	/** The copyright notice to apply to generated files. */
diff --git a/Engine/Plugins/Web/WebAPI/Source/WebAPIEditor/Public/WebAPIEditorSettings.h b/Engine/Plugins/Web/WebAPI/Source/WebAPIEditor/Public/WebAPIEditorSettings.h
--- a/Engine/Plugins/Web/WebAPI/Source/WebAPIEditor/Public/WebAPIEditorSettings.h
+++ b/Engine/Plugins/Web/WebAPI/Source/WebAPIEditor/Public/WebAPIEditorSettings.h
@@ -15,7 +15,7 @@ class WEBAPIEDITOR_API UWebAPIEditorSettings : public UDeveloperSettings
 	GENERATED_BODY()
 	
 public:
-	UPROPERTY(EditAnywhere, BlueprintReadOnly, NoClear, Category = "Generator", meta = (MustImplement = "WebAPICodeGeneratorInterface"))
+	UPROPERTY(EditAnywhere, BlueprintReadOnly, NoClear, Category = "Generator", meta = (MustImplement = "/Script/WebAPIEditor.WebAPICodeGeneratorInterface"))
 	TSoftClassPtr<UObject> CodeGeneratorClass;
 	
 	virtual FName GetContainerName() const override { return FName("Project"); }
