Fix conflicting localization key causing CI warning.

diff --git a/Engine/Plugins/Web/WebAPI/Source/WebAPILiquidJS/Private/WebAPILiquidJSProcess.cpp b/Engine/Plugins/Web/WebAPI/Source/WebAPILiquidJS/Private/WebAPILiquidJSProcess.cpp
--- a/Engine/Plugins/Web/WebAPI/Source/WebAPILiquidJS/Private/WebAPILiquidJSProcess.cpp
+++ b/Engine/Plugins/Web/WebAPI/Source/WebAPILiquidJS/Private/WebAPILiquidJSProcess.cpp
@@ -1,4 +1,4 @@
-﻿// Copyright Epic Games, Inc. All Rights Reserved.
+// Copyright Epic Games, Inc. All Rights Reserved.
 
 #include "WebAPILiquidJSProcess.h"
 
@@ -250,7 +250,7 @@ uint32 FWebAPILiquidJSProcess::Run()
 	{
 		TaskNotification->SetComplete(
 			ErrorTitle,
-			LOCTEXT("WebAPILiquidJS_LaunchFailed", "WebApp exited"),
+			LOCTEXT("WebAPILiquidJS_LaunchFailed2", "WebApp exited"),
 			false
 		);
 	}
