WebAPI plugin moved to experimental

diff --git a/Engine/Plugins/Experimental/Web/WebAPI/Source/WebAPIEditor/Private/Details/Widgets/SWebAPISchemaParameterRow.h b/Engine/Plugins/Experimental/Web/WebAPI/Source/WebAPIEditor/Private/Details/Widgets/SWebAPISchemaParameterRow.h
new file mode 100644
--- /dev/null
+++ b/Engine/Plugins/Experimental/Web/WebAPI/Source/WebAPIEditor/Private/Details/Widgets/SWebAPISchemaParameterRow.h
@@ -0,0 +1,20 @@
+﻿// Copyright Epic Games, Inc. All Rights Reserved.
+
+#pragma once
+
+#include "CoreMinimal.h"
+#include "Details/ViewModels/WebAPIParameterViewModel.h"
+#include "SWebAPISchemaTreeTableRow.h"
+#include "SWebAPISchemaTreeTableRow.inl"
+
+class SWebAPISchemaParameterRow
+	: public SWebAPISchemaTreeTableRow<FWebAPIParameterViewModel>
+{
+public:
+	SLATE_BEGIN_ARGS(SWebAPISchemaParameterRow)
+	{ }
+	SLATE_END_ARGS()
+
+public:
+	void Construct(const FArguments& InArgs, const TSharedRef<FWebAPIParameterViewModel>& InViewModel, const TSharedRef<STableViewBase>& InOwnerTableView);
+};
