Removed some TIsSame usage (part of a larger change to deprecate TIsSame).

diff --git a/Engine/Plugins/Enterprise/DatasmithContent/Source/DatasmithContentEditor/Private/DatasmithSceneDetails.cpp b/Engine/Plugins/Enterprise/DatasmithContent/Source/DatasmithContentEditor/Private/DatasmithSceneDetails.cpp
--- a/Engine/Plugins/Enterprise/DatasmithContent/Source/DatasmithContentEditor/Private/DatasmithSceneDetails.cpp
+++ b/Engine/Plugins/Enterprise/DatasmithContent/Source/DatasmithContentEditor/Private/DatasmithSceneDetails.cpp
@@ -25,7 +25,7 @@ void FDatasmithSceneDetails::CustomizeDetails(IDetailLayoutBuilder& DetailBuilde
 	UDatasmithScene* DatasmithScene = Cast< UDatasmithScene >( Objects[0].Get() );
 	check( DatasmithScene );
 
-	static_assert( TIsSame< typename TRemovePointer< typename TRemoveObjectPointer<decltype( UDatasmithScene::AssetImportData )>::Type >::Type, UDatasmithSceneImportData >::Value, "Please update this details customization" );
+	static_assert( std::is_same_v< typename TRemovePointer< typename TRemoveObjectPointer<decltype( UDatasmithScene::AssetImportData )>::Type >::Type, UDatasmithSceneImportData >, "Please update this details customization" );
 	static_assert( TIsDerivedFrom < UDatasmithSceneImportData, UAssetImportData >::IsDerived, "Please update this details customization" );
 
 	TSharedRef< IPropertyHandle > AssetImportDataPropertyHandle = DetailBuilder.GetProperty( GET_MEMBER_NAME_CHECKED( UDatasmithScene, AssetImportData ) );
diff --git a/Engine/Plugins/Enterprise/DatasmithImporter/Source/DatasmithImporter/Private/ObjectElements/DatasmithUElementsUtils.h b/Engine/Plugins/Enterprise/DatasmithImporter/Source/DatasmithImporter/Private/ObjectElements/DatasmithUElementsUtils.h
--- a/Engine/Plugins/Enterprise/DatasmithImporter/Source/DatasmithImporter/Private/ObjectElements/DatasmithUElementsUtils.h
+++ b/Engine/Plugins/Enterprise/DatasmithImporter/Source/DatasmithImporter/Private/ObjectElements/DatasmithUElementsUtils.h
@@ -48,7 +48,7 @@ public:
 	 * });
 	 */
 	template<typename TargetType, typename Func>
-	static inline typename TEnableIf<TIsSame<TargetType, IDatasmithLevelVariantSetsElement>::Value, bool>::Type ForVariantElement(const TSharedPtr<IDatasmithScene>& RootObject, Func Function)
+	static inline typename TEnableIf<std::is_same_v<TargetType, IDatasmithLevelVariantSetsElement>, bool>::Type ForVariantElement(const TSharedPtr<IDatasmithScene>& RootObject, Func Function)
 	{
 		for (int32 Index = 0; Index < RootObject->GetLevelVariantSetsCount(); ++Index)
 		{
@@ -61,7 +61,7 @@ public:
 		return true;
 	}
 	template<typename TargetType, typename Func>
-	static inline typename TEnableIf<!TIsSame<TargetType, IDatasmithLevelVariantSetsElement>::Value, bool>::Type ForVariantElement(const TSharedPtr<IDatasmithScene>& RootObject, Func Function)
+	static inline typename TEnableIf<!std::is_same_v<TargetType, IDatasmithLevelVariantSetsElement>, bool>::Type ForVariantElement(const TSharedPtr<IDatasmithScene>& RootObject, Func Function)
 	{
 		for (int32 Index = 0; Index < RootObject->GetLevelVariantSetsCount(); ++Index)
 		{
@@ -74,7 +74,7 @@ public:
 		return true;
 	}
 	template<typename TargetType, typename Func>
-	static inline typename TEnableIf<TIsSame<TargetType, IDatasmithVariantSetElement>::Value, bool>::Type ForVariantElement(const TSharedPtr<IDatasmithLevelVariantSetsElement>& RootObject, Func Function)
+	static inline typename TEnableIf<std::is_same_v<TargetType, IDatasmithVariantSetElement>, bool>::Type ForVariantElement(const TSharedPtr<IDatasmithLevelVariantSetsElement>& RootObject, Func Function)
 	{
 		for (int32 Index = 0; Index < RootObject->GetVariantSetsCount(); ++Index)
 		{
@@ -87,7 +87,7 @@ public:
 		return true;
 	}
 	template<typename TargetType, typename Func>
-	static inline typename TEnableIf<!TIsSame<TargetType, IDatasmithVariantSetElement>::Value, bool>::Type ForVariantElement(const TSharedPtr<IDatasmithLevelVariantSetsElement>& RootObject, Func Function)
+	static inline typename TEnableIf<!std::is_same_v<TargetType, IDatasmithVariantSetElement>, bool>::Type ForVariantElement(const TSharedPtr<IDatasmithLevelVariantSetsElement>& RootObject, Func Function)
 	{
 		for (int32 Index = 0; Index < RootObject->GetVariantSetsCount(); ++Index)
 		{
@@ -101,7 +101,7 @@ public:
 	}
 
 	template<typename TargetType, typename Func>
-	static inline typename TEnableIf<TIsSame<TargetType, IDatasmithVariantElement>::Value, bool>::Type ForVariantElement(const TSharedPtr<IDatasmithVariantSetElement>& RootObject, Func Function)
+	static inline typename TEnableIf<std::is_same_v<TargetType, IDatasmithVariantElement>, bool>::Type ForVariantElement(const TSharedPtr<IDatasmithVariantSetElement>& RootObject, Func Function)
 	{
 		for (int32 Index = 0; Index < RootObject->GetVariantsCount(); ++Index)
 		{
@@ -114,7 +114,7 @@ public:
 		return true;
 	}
 	template<typename TargetType, typename Func>
-	static inline typename TEnableIf<!TIsSame<TargetType, IDatasmithVariantElement>::Value, bool>::Type ForVariantElement(const TSharedPtr<IDatasmithVariantSetElement>& RootObject, Func Function)
+	static inline typename TEnableIf<!std::is_same_v<TargetType, IDatasmithVariantElement>, bool>::Type ForVariantElement(const TSharedPtr<IDatasmithVariantSetElement>& RootObject, Func Function)
 	{
 		for (int32 Index = 0; Index < RootObject->GetVariantsCount(); ++Index)
 		{
@@ -127,7 +127,7 @@ public:
 		return true;
 	}
 	template<typename TargetType, typename Func>
-	static inline typename TEnableIf<TIsSame<TargetType, IDatasmithActorBindingElement>::Value, bool>::Type ForVariantElement(const TSharedPtr<IDatasmithVariantElement>& RootObject, Func Function)
+	static inline typename TEnableIf<std::is_same_v<TargetType, IDatasmithActorBindingElement>, bool>::Type ForVariantElement(const TSharedPtr<IDatasmithVariantElement>& RootObject, Func Function)
 	{
 		for (int32 Index = 0; Index < RootObject->GetActorBindingsCount(); ++Index)
 		{
@@ -140,7 +140,7 @@ public:
 		return true;
 	}
 	template<typename TargetType, typename Func>
-	static inline typename TEnableIf<!TIsSame<TargetType, IDatasmithActorBindingElement>::Value, bool>::Type ForVariantElement(const TSharedPtr<IDatasmithVariantElement>& RootObject, Func Function)
+	static inline typename TEnableIf<!std::is_same_v<TargetType, IDatasmithActorBindingElement>, bool>::Type ForVariantElement(const TSharedPtr<IDatasmithVariantElement>& RootObject, Func Function)
 	{
 		for (int32 Index = 0; Index < RootObject->GetActorBindingsCount(); ++Index)
 		{
@@ -153,7 +153,7 @@ public:
 		return true;
 	}
 	template<typename TargetType, typename Func>
-	static inline typename TEnableIf<TIsSame<TargetType, IDatasmithBasePropertyCaptureElement>::Value, bool>::Type ForVariantElement(const TSharedPtr<IDatasmithActorBindingElement>& RootObject, Func Function)
+	static inline typename TEnableIf<std::is_same_v<TargetType, IDatasmithBasePropertyCaptureElement>, bool>::Type ForVariantElement(const TSharedPtr<IDatasmithActorBindingElement>& RootObject, Func Function)
 	{
 		for (int32 Index = 0; Index < RootObject->GetPropertyCapturesCount(); ++Index)
 		{
diff --git a/Engine/Plugins/Enterprise/GLTFExporter/Source/GLTFExporter/Public/Json/GLTFJsonWriter.h b/Engine/Plugins/Enterprise/GLTFExporter/Source/GLTFExporter/Public/Json/GLTFJsonWriter.h
--- a/Engine/Plugins/Enterprise/GLTFExporter/Source/GLTFExporter/Public/Json/GLTFJsonWriter.h
+++ b/Engine/Plugins/Enterprise/GLTFExporter/Source/GLTFExporter/Public/Json/GLTFJsonWriter.h
@@ -68,7 +68,7 @@ public:
 		Write(FGLTFJsonUtilities::GetValue(Enum));
 	}
 
-	template <class ElementType, typename = typename TEnableIf<TNot<TIsSame<ElementType, TCHAR>>::Value>::Type>
+	template <class ElementType, typename = typename TEnableIf<!std::is_same_v<ElementType, TCHAR>>::Type>
 	void Write(const ElementType* Array, SIZE_T ArraySize)
 	{
 		StartArray();
@@ -79,7 +79,7 @@ public:
 		EndArray();
 	}
 
-	template <class ElementType, SIZE_T ArraySize, typename = typename TEnableIf<TNot<TIsSame<ElementType, TCHAR>>::Value>::Type>
+	template <class ElementType, SIZE_T ArraySize, typename = typename TEnableIf<!std::is_same_v<ElementType, TCHAR>>::Type>
 	void Write(const ElementType (&Array)[ArraySize])
 	{
 		Write(Array, ArraySize);
@@ -175,14 +175,14 @@ public:
 		Write(Enum);
 	}
 
-	template <class ElementType, typename = typename TEnableIf<TNot<TIsSame<ElementType, TCHAR>>::Value>::Type>
+	template <class ElementType, typename = typename TEnableIf<!std::is_same_v<ElementType, TCHAR>>::Type>
 	void Write(const FString& Identifier, const ElementType* Array, SIZE_T ArraySize)
 	{
 		SetIdentifier(Identifier);
 		Write(Array, ArraySize);
 	}
 
-	template <class ElementType, SIZE_T ArraySize, typename = typename TEnableIf<TNot<TIsSame<ElementType, TCHAR>>::Value>::Type>
+	template <class ElementType, SIZE_T ArraySize, typename = typename TEnableIf<!std::is_same_v<ElementType, TCHAR>>::Type>
 	void Write(const FString& Identifier, const ElementType (&Array)[ArraySize])
 	{
 		SetIdentifier(Identifier);
diff --git a/Engine/Plugins/Enterprise/GLTFExporter/Source/GLTFExporter/Public/Utilities/GLTFCoreUtilities.h b/Engine/Plugins/Enterprise/GLTFExporter/Source/GLTFExporter/Public/Utilities/GLTFCoreUtilities.h
--- a/Engine/Plugins/Enterprise/GLTFExporter/Source/GLTFExporter/Public/Utilities/GLTFCoreUtilities.h
+++ b/Engine/Plugins/Enterprise/GLTFExporter/Source/GLTFExporter/Public/Utilities/GLTFCoreUtilities.h
@@ -71,13 +71,13 @@ struct GLTFEXPORTER_API FGLTFCoreUtilities
 	template <typename ComponentType>
 	static EGLTFJsonComponentType GetComponentType()
 	{
-		if (TIsSame<ComponentType, int8  >::Value) return EGLTFJsonComponentType::Int8;
-		if (TIsSame<ComponentType, uint8 >::Value) return EGLTFJsonComponentType::UInt8;
-		if (TIsSame<ComponentType, int16 >::Value) return EGLTFJsonComponentType::Int16;
-		if (TIsSame<ComponentType, uint16>::Value) return EGLTFJsonComponentType::UInt16;
-		if (TIsSame<ComponentType, int32 >::Value) return EGLTFJsonComponentType::Int32;
-		if (TIsSame<ComponentType, uint32>::Value) return EGLTFJsonComponentType::UInt32;
-		if (TIsSame<ComponentType, float >::Value) return EGLTFJsonComponentType::Float;
-		return EGLTFJsonComponentType::None;
+		     if constexpr (std::is_same_v<ComponentType, int8  >) return EGLTFJsonComponentType::Int8;
+		else if constexpr (std::is_same_v<ComponentType, uint8 >) return EGLTFJsonComponentType::UInt8;
+		else if constexpr (std::is_same_v<ComponentType, int16 >) return EGLTFJsonComponentType::Int16;
+		else if constexpr (std::is_same_v<ComponentType, uint16>) return EGLTFJsonComponentType::UInt16;
+		else if constexpr (std::is_same_v<ComponentType, int32 >) return EGLTFJsonComponentType::Int32;
+		else if constexpr (std::is_same_v<ComponentType, uint32>) return EGLTFJsonComponentType::UInt32;
+		else if constexpr (std::is_same_v<ComponentType, float >) return EGLTFJsonComponentType::Float;
+		else return EGLTFJsonComponentType::None;
 	}
 };
diff --git a/Engine/Plugins/Experimental/PCG/Source/PCG/Public/Tests/Determinism/PCGDeterminismTestsCommon.h b/Engine/Plugins/Experimental/PCG/Source/PCG/Public/Tests/Determinism/PCGDeterminismTestsCommon.h
--- a/Engine/Plugins/Experimental/PCG/Source/PCG/Public/Tests/Determinism/PCGDeterminismTestsCommon.h
+++ b/Engine/Plugins/Experimental/PCG/Source/PCG/Public/Tests/Determinism/PCGDeterminismTestsCommon.h
@@ -371,7 +371,7 @@ namespace PCGDeterminismTests
 	{
 		check(FirstAttributeBase && SecondAttributeBase);
 
-		if constexpr (TIsSame<FTransform, MetadataAttributeType>::Value)
+		if constexpr (std::is_same_v<FTransform, MetadataAttributeType>)
 		{
 			const FTransform FirstTypedAttribute = static_cast<const FPCGMetadataAttribute<FTransform>*>(FirstAttributeBase)->GetValue(ValueKey);
 			const FTransform SecondTypedAttribute = static_cast<const FPCGMetadataAttribute<FTransform>*>(SecondAttributeBase)->GetValue(ValueKey);
diff --git a/Engine/Plugins/Experimental/Web/WebAPI/Source/WebAPIEditor/Public/WebAPIJsonUtilities.h b/Engine/Plugins/Experimental/Web/WebAPI/Source/WebAPIEditor/Public/WebAPIJsonUtilities.h
--- a/Engine/Plugins/Experimental/Web/WebAPI/Source/WebAPIEditor/Public/WebAPIJsonUtilities.h
+++ b/Engine/Plugins/Experimental/Web/WebAPI/Source/WebAPIEditor/Public/WebAPIJsonUtilities.h
@@ -1,4 +1,4 @@
-﻿// Copyright Epic Games, Inc. All Rights Reserved.
+// Copyright Epic Games, Inc. All Rights Reserved.
 
 #pragma once
 
@@ -18,9 +18,8 @@ namespace UE::Json
 		{
 			template <typename T>
 			auto Requires() -> decltype(
-				TAnd<
-					TIsSame<typename TNumericLimits<T>::NumericType, T>,
-					TNot<TIsSame<T, bool>>>::Value);
+				std::is_same_v<typename TNumericLimits<T>::NumericType, T> &&
+				!std::is_same_v<T, bool>);
 		};
 
 		/** Describes a type that provides a FromJson function. */
@@ -56,10 +55,9 @@ namespace UE::Json
 		struct TIsStringLike<
 				ValueType,
 				typename TEnableIf<
-					TOr<
-						TIsSame<ValueType, FString>,
-						TIsSame<ValueType, FName>,
-						TIsSame<ValueType, FText>>::Value>::Type>
+					std::is_same_v<ValueType, FString> ||
+					std::is_same_v<ValueType, FName> ||
+					std::is_same_v<ValueType, FText>>::Type>
 		{
 			enum
 			{
@@ -280,7 +278,7 @@ namespace UE::Json
 
 	// Bool
 	template <typename ValueType>
-	constexpr typename TEnableIf<TIsSame<ValueType, bool>::Value, void>::Type
+	constexpr typename TEnableIf<std::is_same_v<ValueType, bool>, void>::Type
 	As(const TSharedPtr<FJsonValue>& InJsonValue, ValueType& OutValue)
 	{
 		OutValue = InJsonValue->AsBool();
@@ -315,11 +313,11 @@ namespace UE::Json
 
 	// Bool
 	template <typename ValueType>
-	constexpr typename TEnableIf<TIsSame<ValueType, bool>::Value, bool>::Type
+	constexpr typename TEnableIf<std::is_same_v<ValueType, bool>, bool>::Type
 	TryGet(const TSharedPtr<FJsonValue>& InJsonValue, ValueType& OutValue);
 
 	template <typename ValueType>
-	constexpr typename TEnableIf<TIsSame<ValueType, bool>::Value, bool>::Type
+	constexpr typename TEnableIf<std::is_same_v<ValueType, bool>, bool>::Type
 	TryGetField(const TSharedPtr<FJsonObject>& InJsonObject, const FString& InFieldName, ValueType& OutValue);
 
 	// Array
@@ -385,11 +383,11 @@ namespace UE::Json
 	bool TryGetField(const TSharedPtr<FJsonObject>& InJsonObject, const FString& InFieldName, TJsonReference<ValueType>& OutValue);
 
 	template <typename ValueType>
-	typename TEnableIf<TIsSame<ValueType, TJsonReference<typename ValueType::ElementType>>::Value, bool>::Type
+	typename TEnableIf<std::is_same_v<ValueType, TJsonReference<typename ValueType::ElementType>>, bool>::Type
 	TryGet(const TSharedPtr<FJsonValue>& InJsonValue, ValueType& OutValue);
 
 	template <typename ValueType>
-	typename TEnableIf<TIsSame<ValueType, TJsonReference<typename ValueType::ElementType>>::Value, bool>::Type
+	typename TEnableIf<std::is_same_v<ValueType, TJsonReference<typename ValueType::ElementType>>, bool>::Type
 	TryGetField(const TSharedPtr<FJsonObject>& InJsonObject, const FString& InFieldName, ValueType& OutValue);
 
 	// Object (with FromJson)
@@ -410,27 +408,25 @@ namespace UE::Json
 	// Object (without FromJson)
 	template <typename ValueType>
 	typename TEnableIf<
-		TAnd<
-			TNot<TIsSame<ValueType, TSharedPtr<typename ValueType::ElementType>>>,
-			TNot<TIsSame<ValueType, TJsonReference<typename ValueType::ElementType>>>,
-			TNot<TIsTMap<ValueType>>,
-			TNot<TIsTArray<ValueType>>,
-			TNot<TypeTraits::TIsStringLike<ValueType>>,
-			TNot<TIsPODType<ValueType>>,
-			TNot<TypeTraits::THasFromJson<ValueType>>>::Value, bool>::Type
+		!std::is_same_v<ValueType, TSharedPtr<typename ValueType::ElementType>> &&
+		!std::is_same_v<ValueType, TJsonReference<typename ValueType::ElementType>> &&
+		!TIsTMap<ValueType>::Value &&
+		!TIsTArray<ValueType>::Value &&
+		!TypeTraits::TIsStringLike<ValueType>::Value &&
+		!TIsPODType<ValueType>::Value &&
+		!TypeTraits::THasFromJson<ValueType>::Value, bool>::Type
 	TryGet(const TSharedPtr<FJsonValue>& InJsonValue, ValueType& OutValue);
 
 	// TIsTMap<ContainerType>,
 	template <typename ValueType>
 	typename TEnableIf<
-		TAnd<
-			TNot<TIsSame<ValueType, TSharedPtr<typename ValueType::ElementType>>>,
-			TNot<TIsSame<ValueType, TJsonReference<typename ValueType::ElementType>>>,
-			TNot<TIsTMap<ValueType>>,
-			TNot<TIsTArray<ValueType>>,
-			TNot<TypeTraits::TIsStringLike<ValueType>>,
-			TNot<TIsPODType<ValueType>>,
-			TNot<TypeTraits::THasFromJson<ValueType>>>::Value, bool>::Type
+		!std::is_same_v<ValueType, TSharedPtr<typename ValueType::ElementType>> &&
+		!std::is_same_v<ValueType, TJsonReference<typename ValueType::ElementType>> &&
+		!TIsTMap<ValueType>::Value &&
+		!TIsTArray<ValueType>::Value &&
+		!TypeTraits::TIsStringLike<ValueType>::Value &&
+		!TIsPODType<ValueType>::Value &&
+		!TypeTraits::THasFromJson<ValueType>::Value, bool>::Type
 	TryGetField(const TSharedPtr<FJsonObject>& InJsonObject, const FString& InFieldName, ValueType& OutValue);
 
 	// UniqueObj
@@ -448,11 +444,11 @@ namespace UE::Json
 	bool TryGetField(const TSharedPtr<FJsonObject>& InJsonObject, const FString& InFieldName, TSharedPtr<ValueType>& OutValue);
 
 	template <typename ValueType>
-	typename TEnableIf<TIsSame<ValueType, TSharedPtr<typename ValueType::ElementType>>::Value, bool>::Type
+	typename TEnableIf<std::is_same_v<ValueType, TSharedPtr<typename ValueType::ElementType>>, bool>::Type
 	TryGet(const TSharedPtr<FJsonValue>& InJsonValue, ValueType& OutValue);
 
 	template <typename ValueType>
-	typename TEnableIf<TIsSame<ValueType, TSharedPtr<typename ValueType::ElementType>>::Value, bool>::Type
+	typename TEnableIf<std::is_same_v<ValueType, TSharedPtr<typename ValueType::ElementType>>, bool>::Type
 	TryGetField(const TSharedPtr<FJsonObject>& InJsonObject, const FString& InFieldName, ValueType& OutValue);
 
 	// SharedRef
@@ -495,10 +491,10 @@ namespace UE::Json
 	bool TryGetField(const TSharedPtr<FJsonObject>& InJsonObject, const FString& InFieldName, TOptional<TSharedRef<ValueType>>& OutValue);
 
 	// JsonObject
-	template <typename ValueType, typename = typename TEnableIf<TIsSame<ValueType, TSharedPtr<FJsonObject>>::Value>::Type>
+	template <typename ValueType, typename = typename TEnableIf<std::is_same_v<ValueType, TSharedPtr<FJsonObject>>>::Type>
 	bool TryGet(const TSharedPtr<FJsonValue>& InJsonValue, ValueType& OutValue);
 
-	template <typename ValueType, typename = typename TEnableIf<TIsSame<ValueType, TSharedPtr<FJsonObject>>::Value>::Type>
+	template <typename ValueType, typename = typename TEnableIf<std::is_same_v<ValueType, TSharedPtr<FJsonObject>>>::Type>
 	bool TryGetField(const TSharedPtr<FJsonObject>& InJsonObject, const FString& InFieldName, TSharedPtr<FJsonObject>& OutValue);
 
 	// Variant
diff --git a/Engine/Plugins/Experimental/Web/WebAPI/Source/WebAPIEditor/Public/WebAPIJsonUtilities.inl b/Engine/Plugins/Experimental/Web/WebAPI/Source/WebAPIEditor/Public/WebAPIJsonUtilities.inl
--- a/Engine/Plugins/Experimental/Web/WebAPI/Source/WebAPIEditor/Public/WebAPIJsonUtilities.inl
+++ b/Engine/Plugins/Experimental/Web/WebAPI/Source/WebAPIEditor/Public/WebAPIJsonUtilities.inl
@@ -141,14 +141,14 @@ namespace UE::Json
 
 	// Bool
 	template <typename ValueType>
-	constexpr typename TEnableIf<TIsSame<ValueType, bool>::Value, bool>::Type
+	constexpr typename TEnableIf<std::is_same_v<ValueType, bool>, bool>::Type
 	TryGet(const TSharedPtr<FJsonValue>& InJsonValue, ValueType& OutValue)
 	{
 		return InJsonValue->TryGetBool(OutValue);
 	}
 
 	template <typename ValueType>
-	constexpr typename TEnableIf<TIsSame<ValueType, bool>::Value, bool>::Type
+	constexpr typename TEnableIf<std::is_same_v<ValueType, bool>, bool>::Type
 	TryGetField(const TSharedPtr<FJsonObject>& InJsonObject, const FString& InFieldName, ValueType& OutValue)
 	{
 		return InJsonObject->TryGetBoolField(InFieldName, OutValue);
@@ -333,14 +333,13 @@ namespace UE::Json
 	// Object (without FromJson)
 	template <typename ValueType>
 	typename TEnableIf<
-		TAnd<
-			TNot<TIsSame<ValueType, TSharedPtr<typename ValueType::ElementType>>>,
-			TNot<TIsSame<ValueType, TJsonReference<typename ValueType::ElementType>>>,
-			TNot<TIsTMap<ValueType>>,
-			TNot<TIsTArray<ValueType>>,
-			TNot<TypeTraits::TIsStringLike<ValueType>>,
-			TNot<TIsPODType<ValueType>>,
-			TNot<TypeTraits::THasFromJson<ValueType>>>::Value, bool>::Type
+		!std::is_same_v<ValueType, TSharedPtr<typename ValueType::ElementType>> &&
+		!std::is_same_v<ValueType, TJsonReference<typename ValueType::ElementType>> &&
+		!TIsTMap<ValueType>::Value &&
+		!TIsTArray<ValueType>::Value &&
+		!TypeTraits::TIsStringLike<ValueType>::Value &&
+		!TIsPODType<ValueType>::Value &&
+		!TypeTraits::THasFromJson<ValueType>::Value, bool>::Type
 	TryGet(const TSharedPtr<FJsonValue>& InJsonValue, ValueType& OutValue)
 	{
 		static_assert(TypeTraits::THasFromJson<ValueType>::Value, "ValueType must implement FromJson");
@@ -349,14 +348,13 @@ namespace UE::Json
 
 	template <typename ValueType>
 	typename TEnableIf<
-		TAnd<
-			TNot<TIsSame<ValueType, TSharedPtr<typename ValueType::ElementType>>>,
-			TNot<TIsSame<ValueType, TJsonReference<typename ValueType::ElementType>>>,
-			TNot<TIsTMap<ValueType>>,
-			TNot<TIsTArray<ValueType>>,
-			TNot<TypeTraits::TIsStringLike<ValueType>>,
-			TNot<TIsPODType<ValueType>>,
-			TNot<TypeTraits::THasFromJson<ValueType>>>::Value, bool>::Type
+		!std::is_same_v<ValueType, TSharedPtr<typename ValueType::ElementType>> &&
+		!std::is_same_v<ValueType, TJsonReference<typename ValueType::ElementType>> &&
+		!TIsTMap<ValueType>::Value &&
+		!TIsTArray<ValueType>::Value &&
+		!TypeTraits::TIsStringLike<ValueType>::Value &&
+		!TIsPODType<ValueType>::Value &&
+		!TypeTraits::THasFromJson<ValueType>::Value, bool>::Type
 	TryGetField(const TSharedPtr<FJsonObject>& InJsonObject, const FString& InFieldName, ValueType& OutValue)
 	{
 		static_assert(TypeTraits::THasFromJson<ValueType>::Value, "ValueType must implement FromJson");
@@ -408,14 +406,14 @@ namespace UE::Json
 	}
 
 	template <typename ValueType>
-	typename TEnableIf<TIsSame<ValueType, TJsonReference<typename ValueType::ElementType>>::Value, bool>::Type
+	typename TEnableIf<std::is_same_v<ValueType, TJsonReference<typename ValueType::ElementType>>, bool>::Type
 	TryGet(const TSharedPtr<FJsonValue>& InJsonValue, ValueType& OutValue)
 	{
 		return TryGet<typename ValueType::ElementType>(InJsonValue, OutValue);
 	}
 
 	template <typename ValueType>
-	typename TEnableIf<TIsSame<ValueType, TJsonReference<typename ValueType::ElementType>>::Value, bool>::Type
+	typename TEnableIf<std::is_same_v<ValueType, TJsonReference<typename ValueType::ElementType>>, bool>::Type
 	TryGetField(const TSharedPtr<FJsonObject>& InJsonObject, const FString& InFieldName, ValueType& OutValue)
 	{
 		return TryGet<typename ValueType::ElementType>(InJsonObject, InFieldName, OutValue);
@@ -450,14 +448,14 @@ namespace UE::Json
 	}
 
 	template <typename ValueType>
-	typename TEnableIf<TIsSame<ValueType, TSharedPtr<typename ValueType::ElementType>>::Value, bool>::Type
+	typename TEnableIf<std::is_same_v<ValueType, TSharedPtr<typename ValueType::ElementType>>, bool>::Type
 	TryGet(const TSharedPtr<FJsonValue>& InJsonValue, ValueType& OutValue)
 	{
 		return TryGet<typename ValueType::ElementType>(InJsonValue, OutValue);
 	}
 
 	template <typename ValueType>
-	typename TEnableIf<TIsSame<ValueType, TSharedPtr<typename ValueType::ElementType>>::Value, bool>::Type
+	typename TEnableIf<std::is_same_v<ValueType, TSharedPtr<typename ValueType::ElementType>>, bool>::Type
 	TryGetField(const TSharedPtr<FJsonObject>& InJsonObject, const FString& InFieldName, ValueType& OutValue)
 	{
 		return TryGet<typename ValueType::ElementType>(InJsonObject, InFieldName, OutValue);
diff --git a/Engine/Plugins/Experimental/Web/WebAPI/Source/WebAPIOpenAPI/Private/Tests/WebAPIOpenAPI3.spec.cpp b/Engine/Plugins/Experimental/Web/WebAPI/Source/WebAPIOpenAPI/Private/Tests/WebAPIOpenAPI3.spec.cpp
--- a/Engine/Plugins/Experimental/Web/WebAPI/Source/WebAPIOpenAPI/Private/Tests/WebAPIOpenAPI3.spec.cpp
+++ b/Engine/Plugins/Experimental/Web/WebAPI/Source/WebAPIOpenAPI/Private/Tests/WebAPIOpenAPI3.spec.cpp
@@ -53,11 +53,11 @@ BEGIN_DEFINE_SPEC(FWebAPIOpenAPI3Spec,
 	}
 
 	template <class DestinationType, class SourceType>
-	typename TEnableIf<!TIsSame<DestinationType, SourceType>::Value, void>::Type
+	typename TEnableIf<!std::is_same_v<DestinationType, SourceType>, void>::Type
 	TryAssign(TSharedPtr<DestinationType>& InDst, const TSharedPtr<SourceType>& InSrc) { }
 
 	template <class DestinationType, class SourceType>
-	typename TEnableIf<TIsSame<DestinationType, SourceType>::Value, void>::Type
+	typename TEnableIf<std::is_same_v<DestinationType, SourceType>, void>::Type
 	TryAssign(TSharedPtr<DestinationType>& InDst, const TSharedPtr<SourceType>& InSrc)
 	{
 		InDst = InSrc;
diff --git a/Engine/Plugins/FX/Niagara/Source/Niagara/Public/NiagaraTypes.h b/Engine/Plugins/FX/Niagara/Source/Niagara/Public/NiagaraTypes.h
--- a/Engine/Plugins/FX/Niagara/Source/Niagara/Public/NiagaraTypes.h
+++ b/Engine/Plugins/FX/Niagara/Source/Niagara/Public/NiagaraTypes.h
@@ -1235,18 +1235,19 @@ struct TStructOpsTypeTraits<FNiagaraTypeDefinition> : public TStructOpsTypeTrait
 template<typename T>
 const FNiagaraTypeDefinition& FNiagaraTypeDefinition::Get()
 {
-	if (TIsSame<T, float>::Value) { return GetFloatDef(); }
-	if (TIsSame<T, FVector2f>::Value) { return GetVec2Def(); }
-	if (TIsSame<T, FVector3f>::Value) { return GetVec3Def(); }	
-	if (TIsSame<T, FNiagaraPosition>::Value) { return GetPositionDef(); }	
-	if (TIsSame<T, FVector4f>::Value) { return GetVec4Def(); }
-	if (TIsSame<T, int32>::Value) { return GetIntDef(); }
-	if (TIsSame<T, FNiagaraBool>::Value) { return GetBoolDef(); }
-	if (TIsSame<T, FQuat4f>::Value) { return GetQuatDef(); }
-	if (TIsSame<T, FMatrix44f>::Value) { return GetMatrix4Def(); }
-	if (TIsSame<T, FLinearColor>::Value) { return GetColorDef(); }
-	if (TIsSame<T, FNiagaraID>::Value) { return GetIDDef(); }
-	if (TIsSame<T, FNiagaraRandInfo>::Value) { return GetRandInfoDef(); }
+	     if constexpr (std::is_same_v<T, float>) { return GetFloatDef(); }
+	else if constexpr (std::is_same_v<T, FVector2f>) { return GetVec2Def(); }
+	else if constexpr (std::is_same_v<T, FVector3f>) { return GetVec3Def(); }
+	else if constexpr (std::is_same_v<T, FNiagaraPosition>) { return GetPositionDef(); }
+	else if constexpr (std::is_same_v<T, FVector4f>) { return GetVec4Def(); }
+	else if constexpr (std::is_same_v<T, int32>) { return GetIntDef(); }
+	else if constexpr (std::is_same_v<T, FNiagaraBool>) { return GetBoolDef(); }
+	else if constexpr (std::is_same_v<T, FQuat4f>) { return GetQuatDef(); }
+	else if constexpr (std::is_same_v<T, FMatrix44f>) { return GetMatrix4Def(); }
+	else if constexpr (std::is_same_v<T, FLinearColor>) { return GetColorDef(); }
+	else if constexpr (std::is_same_v<T, FNiagaraID>) { return GetIDDef(); }
+	else if constexpr (std::is_same_v<T, FNiagaraRandInfo>) { return GetRandInfoDef(); }
+	else { static_assert(sizeof(T) == 0, "Unsupported type"); }
 }
 
 //////////////////////////////////////////////////////////////////////////
diff --git a/Engine/Plugins/Online/OnlineSubsystemTencent/Source/Private/OnlineAsyncTasksTencent.cpp b/Engine/Plugins/Online/OnlineSubsystemTencent/Source/Private/OnlineAsyncTasksTencent.cpp
--- a/Engine/Plugins/Online/OnlineSubsystemTencent/Source/Private/OnlineAsyncTasksTencent.cpp
+++ b/Engine/Plugins/Online/OnlineSubsystemTencent/Source/Private/OnlineAsyncTasksTencent.cpp
@@ -251,7 +251,7 @@ namespace
 		BEGIN_JSON_SERIALIZER
 			if (Serializer.IsSaving())
 			{
-				static_assert(TIsSame<decltype(RailAssetInfo.asset_id), uint64>::Value, "Change how we format asset_id");
+				static_assert(std::is_same_v<decltype(RailAssetInfo.asset_id), uint64>, "Change how we format asset_id");
 				FString AssetId = FString::Printf(TEXT("%llu"), RailAssetInfo.asset_id);
 				JSON_SERIALIZE("asset_id", AssetId);
 				JSON_SERIALIZE("product_id", RailAssetInfo.product_id);
@@ -298,8 +298,8 @@ namespace
 			FPurchaseReceipt::FLineItemInfo LineItemInfo;
 
 			const rail::RailAssetInfo& AssetInfo(AssetInfos[AssetIndex]);
-			static_assert(TIsSame<decltype(AssetInfo.product_id), uint32>::Value, "Change how we construct string from AssetInfo.product_id");
-			static_assert(TIsSame<decltype(AssetInfo.asset_id), uint64>::Value, "Change how we construct string from AssetInfo.asset_id");
+			static_assert(std::is_same_v<decltype(AssetInfo.product_id), uint32>, "Change how we construct string from AssetInfo.product_id");
+			static_assert(std::is_same_v<decltype(AssetInfo.asset_id), uint64>, "Change how we construct string from AssetInfo.asset_id");
 			LineItemInfo.ItemName = FString::Printf(TEXT("%u"), AssetInfo.product_id);
 			LineItemInfo.UniqueId = FString::Printf(TEXT("%llu"), AssetInfo.asset_id);
 			if (AssetInfo.state == rail::EnumRailAssetState::kRailAssetStateNormal)
@@ -2392,7 +2392,7 @@ void FOnlineAsyncTaskRailRequestAllPurchasableProducts::OnRailRequestAllPurchasa
 			const rail::RailPurchaseProductInfo& ProductInfo(Response->purchasable_products[Index]);
 			
 			FOnlineStoreOfferRef StoreOffer = MakeShared<FOnlineStoreOffer>();
-			static_assert(TIsSame<decltype(ProductInfo.product_id), uint32>::Value, "Change how we construct string from ProductInfo.product_id");
+			static_assert(std::is_same_v<decltype(ProductInfo.product_id), uint32>, "Change how we construct string from ProductInfo.product_id");
 			StoreOffer->OfferId = FString::Printf(TEXT("%u"), ProductInfo.product_id);
 			StoreOffer->Title = FText::FromString(LexToString(ProductInfo.name));
 			StoreOffer->Description = FText::FromString(LexToString(ProductInfo.description));
diff --git a/Engine/Plugins/Runtime/Metasound/Source/MetasoundFrontend/Public/MetasoundEnum.h b/Engine/Plugins/Runtime/Metasound/Source/MetasoundFrontend/Public/MetasoundEnum.h
--- a/Engine/Plugins/Runtime/Metasound/Source/MetasoundFrontend/Public/MetasoundEnum.h
+++ b/Engine/Plugins/Runtime/Metasound/Source/MetasoundFrontend/Public/MetasoundEnum.h
@@ -113,7 +113,7 @@ namespace Metasound
 		static_assert(TIsEnum<EnumType>::Value, "Expecting an Enum type");
 		
 		// Allow Non-enum class Enum or Enum class that's are derived from int32
-		static_assert(TIsSame<typename std::underlying_type<EnumType>::type, int32>::Value || 
+		static_assert(std::is_same_v<typename std::underlying_type<EnumType>::type, int32> ||
 			!TIsEnumClass<EnumType>::Value, "We serialize to int32, so limit to that for now");
 
 		// Default.
diff --git a/Engine/Plugins/Runtime/NetworkPrediction/Source/NetworkPrediction/Public/Services/NetworkPredictionService_Ticking.inl b/Engine/Plugins/Runtime/NetworkPrediction/Source/NetworkPrediction/Public/Services/NetworkPredictionService_Ticking.inl
--- a/Engine/Plugins/Runtime/NetworkPrediction/Source/NetworkPrediction/Public/Services/NetworkPredictionService_Ticking.inl
+++ b/Engine/Plugins/Runtime/NetworkPrediction/Source/NetworkPrediction/Public/Services/NetworkPredictionService_Ticking.inl
@@ -20,7 +20,7 @@ struct TTickUtil
 	using FrameDataType = typename TInstanceFrameState<ModelDef>::FFrame;
 
 	template<typename SimulationType = typename ModelDef::Simulation>
-	static typename TEnableIf<!TIsSame<SimulationType, void>::Value>::Type DoTick(TInstanceData<ModelDef>& Instance, FrameDataType& InputFrameData, FrameDataType& OutputFrameData, const FNetSimTimeStep& Step, const int32 CueTimeMS, ESimulationTickContext TickContext)
+	static typename TEnableIf<!std::is_same_v<SimulationType, void>>::Type DoTick(TInstanceData<ModelDef>& Instance, FrameDataType& InputFrameData, FrameDataType& OutputFrameData, const FNetSimTimeStep& Step, const int32 CueTimeMS, ESimulationTickContext TickContext)
 	{
 		Instance.CueDispatcher->PushContext({Step.Frame, CueTimeMS, TickContext});
 
@@ -49,7 +49,7 @@ struct TTickUtil
 	}
 
 	template<typename SimulationType = typename ModelDef::Simulation>
-	static typename TEnableIf<TIsSame<SimulationType, void>::Value>::Type DoTick(TInstanceData<ModelDef>& Instance, FrameDataType& InputFrameData, FrameDataType& OutputFrameData, const FNetSimTimeStep& Step, const int32 EndTimeMS, ESimulationTickContext TickContext)
+	static typename TEnableIf<std::is_same_v<SimulationType, void>>::Type DoTick(TInstanceData<ModelDef>& Instance, FrameDataType& InputFrameData, FrameDataType& OutputFrameData, const FNetSimTimeStep& Step, const int32 EndTimeMS, ESimulationTickContext TickContext)
 	{
 		npCheckf(false, TEXT("DoTick called on %s with no Simulation defined"), ModelDef::GetName());
 	}
diff --git a/Engine/Plugins/Runtime/ReplicationSystemTestPlugin/Source/Private/NetworkAutomationTestMacros.h b/Engine/Plugins/Runtime/ReplicationSystemTestPlugin/Source/Private/NetworkAutomationTestMacros.h
--- a/Engine/Plugins/Runtime/ReplicationSystemTestPlugin/Source/Private/NetworkAutomationTestMacros.h
+++ b/Engine/Plugins/Runtime/ReplicationSystemTestPlugin/Source/Private/NetworkAutomationTestMacros.h
@@ -84,7 +84,7 @@ Name(const T& Value, const T& ExpectedValue, const char* ValueText, const char*
 		return UE::Net::CreateTestFailure() << TEXT("Expected ") << ValueText << TEXT(' ') << TEXT(#Operator) <<  TEXT(' ') << ExpectedValueText << TEXT(". ") << Value << TEXT(" was compared to ") << ExpectedValue << TEXT(". "); \
 	} \
 } \
-template<typename T, typename U, typename TEnableIf<TIsPointer<T>::Value && TIsSame<U, nullptr_t>::Value, U>::Type X = nullptr> \
+template<typename T, typename U, typename TEnableIf<TIsPointer<T>::Value && std::is_same_v<U, nullptr_t>, U>::Type X = nullptr> \
 UE::Net::FTestResult \
 Name(const T& Value, U ExpectedValue, const char* ValueText, const char* ExpectedValueText) \
 { \
diff --git a/Engine/Plugins/Runtime/ReplicationSystemTestPlugin/Source/Private/TestMessage.h b/Engine/Plugins/Runtime/ReplicationSystemTestPlugin/Source/Private/TestMessage.h
--- a/Engine/Plugins/Runtime/ReplicationSystemTestPlugin/Source/Private/TestMessage.h
+++ b/Engine/Plugins/Runtime/ReplicationSystemTestPlugin/Source/Private/TestMessage.h
@@ -78,14 +78,14 @@ inline FTestMessage& operator<<(FTestMessage& Message, const void* Value)
 }
 
 // uintptr_t streaming
-template<typename T, typename = typename TEnableIf<TIsSame<T, uintptr_t>::Value && !TIsSame<uintptr_t, unsigned>::Value && !TIsSame<uintptr_t, uint64>::Value, uintptr_t>::Type>
+template<typename T, typename = typename TEnableIf<std::is_same_v<T, uintptr_t> && !std::is_same_v<uintptr_t, unsigned> && !std::is_same_v<uintptr_t, uint64>, uintptr_t>::Type>
 inline FTestMessage& operator<<(FTestMessage& Message, T Value)
 {
 	return Message << TStringBuilder<32>().Appendf(TEXT("0x%" UPTRINT_X_FMT), UPTRINT(Value)).ToString();
 }
 
 // ptrdiff_t streaming
-template<typename T, typename = typename TEnableIf<TIsSame<T, ptrdiff_t>::Value, ptrdiff_t>::Type, typename = typename TEnableIf<TIsSame<T, ptrdiff_t>::Value && !TIsSame<ptrdiff_t, int32>::Value && !TIsSame<ptrdiff_t, int64>::Value, ptrdiff_t>::Type>
+template<typename T, typename = typename TEnableIf<std::is_same_v<T, ptrdiff_t>, ptrdiff_t>::Type, typename = typename TEnableIf<std::is_same_v<T, ptrdiff_t> && !std::is_same_v<ptrdiff_t, int32> && !std::is_same_v<ptrdiff_t, int64>, ptrdiff_t>::Type>
 inline FTestMessage& operator<<(FTestMessage& Message, T Value)
 {
 	return Message << TStringBuilder<32>().Appendf(TEXT("%" SSIZE_T_FMT), SSIZE_T(Value)).ToString();
diff --git a/Engine/Plugins/Slate/SlateInsights/Source/SlateInsights/Private/SlateProvider.h b/Engine/Plugins/Slate/SlateInsights/Source/SlateInsights/Private/SlateProvider.h
--- a/Engine/Plugins/Slate/SlateInsights/Source/SlateInsights/Private/SlateProvider.h
+++ b/Engine/Plugins/Slate/SlateInsights/Source/SlateInsights/Private/SlateProvider.h
@@ -33,7 +33,7 @@ private:
 
 public:
 	constexpr FWidgetId() : Value(0) {}
-	template<typename T, typename U = typename TEnableIf<TIsSame<T, uint64>::Value>::Type>
+	template<typename T, typename U = typename TEnableIf<std::is_same_v<T, uint64>>::Type>
 	constexpr FWidgetId(T InValue) : Value(InValue) {}
 	explicit operator bool() const { return Value != 0; }
 	uint64 GetValue() const { return Value; }
diff --git a/Engine/Plugins/VirtualProduction/DMX/DMXProtocol/Source/DMXProtocol/Public/DMXConversions.h b/Engine/Plugins/VirtualProduction/DMX/DMXProtocol/Source/DMXProtocol/Public/DMXConversions.h
--- a/Engine/Plugins/VirtualProduction/DMX/DMXProtocol/Source/DMXProtocol/Public/DMXConversions.h
+++ b/Engine/Plugins/VirtualProduction/DMX/DMXProtocol/Source/DMXProtocol/Public/DMXConversions.h
@@ -20,7 +20,7 @@ class DMXPROTOCOL_API FDMXConversions
 
 public:
 	/** Converts a uint32 to a byte array */
-	template <typename T, typename U = typename TEnableIf<TIsSame<T, uint32>::Value>::Type>
+	template <typename T, typename U = typename TEnableIf<std::is_same_v<T, uint32>>::Type>
 	static TArray<uint8> UnsignedInt32ToByteArray(T Value, EDMXFixtureSignalFormat SignalFormat, bool bLSBOrder)
 	{
 		const uint8 NumBytes = GetSizeOfSignalFormat(SignalFormat);
@@ -42,7 +42,7 @@ public:
 	}
 
 	/** Converts a normalized value to a byte array. Normalized value has to be in the 0-1 range. Assumes max signal format is 24bit. */
-	template <typename T, typename U = typename TEnableIf<TIsSame<T, float>::Value>::Type>
+	template <typename T, typename U = typename TEnableIf<std::is_same_v<T, float>>::Type>
 	static TArray<uint8> NormalizedDMXValueToByteArray(T NormalizedValue, EDMXFixtureSignalFormat SignalFormat, bool bLSBOrder)
 	{
 		NormalizedValue = FMath::Clamp(NormalizedValue, 0.f, 1.f);
diff --git a/Engine/Plugins/VirtualProduction/RemoteControl/Source/RemoteControl/Public/RemoteControlProtocolBinding.h b/Engine/Plugins/VirtualProduction/RemoteControl/Source/RemoteControl/Public/RemoteControlProtocolBinding.h
--- a/Engine/Plugins/VirtualProduction/RemoteControl/Source/RemoteControl/Public/RemoteControlProtocolBinding.h
+++ b/Engine/Plugins/VirtualProduction/RemoteControl/Source/RemoteControl/Public/RemoteControlProtocolBinding.h
@@ -87,7 +87,7 @@ public:
 
 	/** Get Binding Range Struct Value */
 	template <typename ValueType, typename PropertyType>
-	typename TEnableIf<TIsSame<FStructProperty, PropertyType>::Value, ValueType>::Type
+	typename TEnableIf<std::is_same_v<FStructProperty, PropertyType>, ValueType>::Type
 	GetRangeValue()
 	{
 		check(InterpolationRangePropertyData.Num() && InterpolationRangePropertyData.Num() == sizeof(ValueType))
@@ -105,7 +105,7 @@ public:
 
 	/** Set Binding Range Struct Value based on templated value input */
 	template <typename ValueType, typename PropertyType>
-	typename TEnableIf<TIsSame<FStructProperty, PropertyType>::Value, void>::Type
+	typename TEnableIf<std::is_same_v<FStructProperty, PropertyType>, void>::Type
 	SetRangeValue(ValueType InRangeValue)
 	{
 		check(InterpolationRangePropertyData.Num() && InterpolationRangePropertyData.Num() == sizeof(ValueType));
@@ -135,7 +135,7 @@ public:
 
 	/** Get Mapping Property Struct Value as a primitive type */
 	template <typename ValueType, typename PropertyType>
-	typename TEnableIf<TIsSame<FStructProperty, PropertyType>::Value, ValueType>::Type
+	typename TEnableIf<std::is_same_v<FStructProperty, PropertyType>, ValueType>::Type
 	GetMappingValueAsPrimitive()
 	{
 		check(InterpolationMappingPropertyData.Num() && InterpolationMappingPropertyData.Num() == sizeof(ValueType));
@@ -155,7 +155,7 @@ public:
 
 	/** Set primitive Mapping Property Struct Value based on template value input */
 	template <typename ValueType, typename PropertyType>
-	typename TEnableIf<TIsSame<FStructProperty, PropertyType>::Value, void>::Type
+	typename TEnableIf<std::is_same_v<FStructProperty, PropertyType>, void>::Type
 	SetMappingValueAsPrimitive(ValueType InMappingPropertyValue)
 	{
 		check(InterpolationMappingPropertyData.Num() && InterpolationMappingPropertyData.Num() == sizeof(ValueType));
diff --git a/Engine/Plugins/VirtualProduction/RemoteControl/Source/RemoteControlCommon/Public/RCPropertyUtilities.h b/Engine/Plugins/VirtualProduction/RemoteControl/Source/RemoteControlCommon/Public/RCPropertyUtilities.h
--- a/Engine/Plugins/VirtualProduction/RemoteControl/Source/RemoteControlCommon/Public/RCPropertyUtilities.h
+++ b/Engine/Plugins/VirtualProduction/RemoteControl/Source/RemoteControlCommon/Public/RCPropertyUtilities.h
@@ -116,25 +116,23 @@ namespace RemoteControlPropertyUtilities
 	/** Reads the raw data from InSrc and deserializes to OutDst. */
 	template <typename PropertyType>
 	typename TEnableIf<
-		TOr<
-			TAnd<
-				TIsDerivedFrom<PropertyType, FProperty>,
-				TNot<TIsSame<PropertyType, FProperty>>,
-				TNot<TIsSame<PropertyType, FNumericProperty>>>,
-			TIsSame<PropertyType, FEnumProperty>
-			>::Value, bool>::Type
+		(
+			TIsDerivedFrom<PropertyType, FProperty>::Value &&
+			!std::is_same_v<PropertyType, FProperty> &&
+			!std::is_same_v<PropertyType, FNumericProperty>
+		) ||
+		std::is_same_v<PropertyType, FEnumProperty>, bool>::Type
 	Deserialize(const FRCPropertyVariant& InSrc, FRCPropertyVariant& OutDst);
 	
 	/** Reads the property value from InSrc and serializes to OutDst. */
 	template <typename PropertyType>
 	typename TEnableIf<
-		TOr<
-			TAnd<
-				TIsDerivedFrom<PropertyType, FProperty>,
-				TNot<TIsSame<PropertyType, FProperty>>,
-				TNot<TIsSame<PropertyType, FNumericProperty>>>,
-			TIsSame<PropertyType, FEnumProperty>
-			>::Value, bool>::Type
+		(
+			TIsDerivedFrom<PropertyType, FProperty>::Value &&
+			!std::is_same_v<PropertyType, FProperty> &&
+			!std::is_same_v<PropertyType, FNumericProperty>
+		) ||
+		std::is_same_v<PropertyType, FEnumProperty>, bool>::Type
 	Serialize(const FRCPropertyVariant& InSrc, FRCPropertyVariant& OutDst);
 
 	/** Specialization for FBoolProperty. */
@@ -148,17 +146,15 @@ namespace RemoteControlPropertyUtilities
 	/** Specialization for FProperty casts and forwards to specializations. */
 	template <typename PropertyType>
 	typename TEnableIf<
-		TOr<
-			TIsSame<PropertyType, FProperty>,
-			TIsSame<PropertyType, FNumericProperty>>::Value, bool>::Type
+			std::is_same_v<PropertyType, FProperty> ||
+			std::is_same_v<PropertyType, FNumericProperty>, bool>::Type
 	Deserialize(const FRCPropertyVariant& InSrc, FRCPropertyVariant& OutDst);
 	
 	/** Specialization for FProperty casts and forwards to specializations. */
 	template <typename PropertyType>
 	typename TEnableIf<
-		TOr<
-			TIsSame<PropertyType, FProperty>,
-			TIsSame<PropertyType, FNumericProperty>>::Value, bool>::Type
+		std::is_same_v<PropertyType, FProperty> ||
+		std::is_same_v<PropertyType, FNumericProperty>, bool>::Type
 	Serialize(const FRCPropertyVariant& InSrc, FRCPropertyVariant& OutDst);
 
 	static FProperty* FindSetterArgument(UFunction* SetterFunction, FProperty* PropertyToModify);
diff --git a/Engine/Plugins/VirtualProduction/RemoteControl/Source/RemoteControlCommon/Public/RCPropertyUtilities.inl b/Engine/Plugins/VirtualProduction/RemoteControl/Source/RemoteControlCommon/Public/RCPropertyUtilities.inl
--- a/Engine/Plugins/VirtualProduction/RemoteControl/Source/RemoteControlCommon/Public/RCPropertyUtilities.inl
+++ b/Engine/Plugins/VirtualProduction/RemoteControl/Source/RemoteControlCommon/Public/RCPropertyUtilities.inl
@@ -1,4 +1,4 @@
-﻿// Copyright Epic Games, Inc. All Rights Reserved.
+// Copyright Epic Games, Inc. All Rights Reserved.
 
 #include "CoreMinimal.h"
 
@@ -206,13 +206,12 @@ namespace RemoteControlPropertyUtilities
 	/** Reads the raw data from InSrc and deserializes to OutDst. */
 	template <typename PropertyType>
 	typename TEnableIf<
-		TOr<
-			TAnd<
-				TIsDerivedFrom<PropertyType, FProperty>,
-				TNot<TIsSame<PropertyType, FProperty>>,
-				TNot<TIsSame<PropertyType, FNumericProperty>>>,
-			TIsSame<PropertyType, FEnumProperty>
-			>::Value, bool>::Type
+		(
+			TIsDerivedFrom<PropertyType, FProperty>::Value &&
+			!std::is_same_v<PropertyType, FProperty> &&
+			!std::is_same_v<PropertyType, FNumericProperty>
+		) ||
+		std::is_same_v<PropertyType, FEnumProperty>, bool>::Type
 	Deserialize(const FRCPropertyVariant& InSrc, FRCPropertyVariant& OutDst)
 	{
 		using ValueType = typename TRemoteControlPropertyTypeTraits<PropertyType>::ValueType;
@@ -275,13 +274,12 @@ namespace RemoteControlPropertyUtilities
 	/** Reads the property value from InSrc and serializes to OutDst. */
 	template <typename PropertyType>
 	typename TEnableIf<
-		TOr<
-			TAnd<
-				TIsDerivedFrom<PropertyType, FProperty>,
-				TNot<TIsSame<PropertyType, FProperty>>,
-				TNot<TIsSame<PropertyType, FNumericProperty>>>,
-			TIsSame<PropertyType, FEnumProperty>
-			>::Value, bool>::Type
+		(
+			TIsDerivedFrom<PropertyType, FProperty>::Value &&
+			!std::is_same_v<PropertyType, FProperty> &&
+			!std::is_same_v<PropertyType, FNumericProperty>
+		) ||
+		std::is_same_v<PropertyType, FEnumProperty>, bool>::Type
 	Serialize(const FRCPropertyVariant& InSrc, FRCPropertyVariant& OutDst)
 	{
 		using ValueType = typename TRemoteControlPropertyTypeTraits<PropertyType>::ValueType;
@@ -341,9 +339,8 @@ namespace RemoteControlPropertyUtilities
 	/** Specialization for FProperty casts and forwards to specializations. */
 	template <typename PropertyType>
 	typename TEnableIf<
-		TOr<
-			TIsSame<PropertyType, FProperty>,
-			TIsSame<PropertyType, FNumericProperty>>::Value, bool>::Type
+		std::is_same_v<PropertyType, FProperty> ||
+		std::is_same_v<PropertyType, FNumericProperty>, bool>::Type
 	Deserialize(const FRCPropertyVariant& InSrc, FRCPropertyVariant& OutDst)
 	{
 		const FProperty* Property = OutDst.GetProperty();
@@ -355,9 +352,8 @@ namespace RemoteControlPropertyUtilities
 	/** Specialization for FProperty casts and forwards to specializations. */
 	template <typename PropertyType>
 	typename TEnableIf<
-		TOr<
-			TIsSame<PropertyType, FProperty>,
-			TIsSame<PropertyType, FNumericProperty>>::Value, bool>::Type
+		std::is_same_v<PropertyType, FProperty> ||
+		std::is_same_v<PropertyType, FNumericProperty>, bool>::Type
 	Serialize(const FRCPropertyVariant& InSrc, FRCPropertyVariant& OutDst)
 	{
 		const FProperty* Property = InSrc.GetProperty();
diff --git a/Engine/Plugins/VirtualProduction/RemoteControl/Source/RemoteControlCommon/Public/RCTypeTraits.h b/Engine/Plugins/VirtualProduction/RemoteControl/Source/RemoteControlCommon/Public/RCTypeTraits.h
--- a/Engine/Plugins/VirtualProduction/RemoteControl/Source/RemoteControlCommon/Public/RCTypeTraits.h
+++ b/Engine/Plugins/VirtualProduction/RemoteControl/Source/RemoteControlCommon/Public/RCTypeTraits.h
@@ -1,4 +1,4 @@
-﻿// Copyright Epic Games, Inc. All Rights Reserved.
+// Copyright Epic Games, Inc. All Rights Reserved.
 
 #pragma once
 
@@ -20,9 +20,8 @@ namespace RemoteControlTypeTraits
 		{
 			template <typename T>
 			auto Requires() -> decltype(
-				TAnd<
-					TIsSame<typename TNumericLimits<T>::NumericType, T>,
-					TNot<TIsSame<T, bool>>>::Value);
+				std::is_same_v<typename TNumericLimits<T>::NumericType, T> &&
+				!std::is_same_v<T, bool>);
 		};
 
 		/** Concept to check if PropertyType::TCppType matches ValueType */
@@ -36,7 +35,7 @@ namespace RemoteControlTypeTraits
 		struct CPropertyHasTCppType
 		{
 			template <typename PropertyType>
-			auto Requires() -> decltype(TIsSame<typename PropertyType::TCppType, void>::Value);
+			auto Requires() -> decltype(std::is_same_v<typename PropertyType::TCppType, void>);
 		};
 	}
 
@@ -51,10 +50,9 @@ namespace RemoteControlTypeTraits
 	struct TIsStringLikeProperty<
 			PropertyType,
 			typename TEnableIf<
-				TOr<
-					TIsSame<PropertyType, FStrProperty>,
-					TIsSame<PropertyType, FNameProperty>,
-					TIsSame<PropertyType, FTextProperty>>::Value>::Type>
+					std::is_same_v<PropertyType, FStrProperty> ||
+					std::is_same_v<PropertyType, FNameProperty> ||
+					std::is_same_v<PropertyType, FTextProperty>>::Type>
 	{
 		enum
 		{
@@ -73,10 +71,9 @@ namespace RemoteControlTypeTraits
 	struct TIsStringLikeValue<
 			ValueType,
 			typename TEnableIf<
-				TOr<
-					TIsSame<ValueType, FString>,
-					TIsSame<ValueType, FName>,
-					TIsSame<ValueType, FText>>::Value>::Type>
+					std::is_same_v<ValueType, FString> ||
+					std::is_same_v<ValueType, FName> ||
+					std::is_same_v<ValueType, FText>>::Type>
 	{
 		enum
 		{
diff --git a/Engine/Plugins/VirtualProduction/RemoteControl/Source/RemoteControlCommon/Public/RCTypeUtilities.h b/Engine/Plugins/VirtualProduction/RemoteControl/Source/RemoteControlCommon/Public/RCTypeUtilities.h
--- a/Engine/Plugins/VirtualProduction/RemoteControl/Source/RemoteControlCommon/Public/RCTypeUtilities.h
+++ b/Engine/Plugins/VirtualProduction/RemoteControl/Source/RemoteControlCommon/Public/RCTypeUtilities.h
@@ -441,7 +441,7 @@ namespace RemoteControlTypeUtilities
 	}
 
 	template <typename ValueType>
-	static typename TEnableIf<TIsSame<ValueType, FStructOnScope>::Value, TSharedPtr<ValueType>>::Type
+	static typename TEnableIf<std::is_same_v<ValueType, FStructOnScope>, TSharedPtr<ValueType>>::Type
 	GetDefaultMappingValueMin(const FStructProperty* InProperty)
 	{
 		check(InProperty);
@@ -493,7 +493,7 @@ namespace RemoteControlTypeUtilities
 	}
 
 	template <typename ValueType>
-	static typename TEnableIf<TIsSame<ValueType, FStructOnScope>::Value, TSharedPtr<ValueType>>::Type
+	static typename TEnableIf<std::is_same_v<ValueType, FStructOnScope>, TSharedPtr<ValueType>>::Type
 	GetDefaultMappingValueMax(const FStructProperty* InProperty)
 	{
 		check(InProperty);
