Replaced TIsTArray trait class with TIsTArray_V.

diff --git a/Engine/Plugins/Experimental/Web/WebAPI/Source/WebAPIEditor/Public/WebAPIJsonUtilities.h b/Engine/Plugins/Experimental/Web/WebAPI/Source/WebAPIEditor/Public/WebAPIJsonUtilities.h
--- a/Engine/Plugins/Experimental/Web/WebAPI/Source/WebAPIEditor/Public/WebAPIJsonUtilities.h
+++ b/Engine/Plugins/Experimental/Web/WebAPI/Source/WebAPIEditor/Public/WebAPIJsonUtilities.h
@@ -322,11 +322,11 @@ namespace UE::Json
 
 	// Array
 	template <typename ContainerType>
-	typename TEnableIf<TIsTArray<ContainerType>::Value, bool>::Type
+	typename TEnableIf<TIsTArray_V<ContainerType>, bool>::Type
 	TryGet(const TSharedPtr<FJsonValue>& InJsonValue, ContainerType& OutValues);
 
 	template <typename ContainerType>
-	constexpr typename TEnableIf<TIsTArray<ContainerType>::Value, bool>::Type
+	constexpr typename TEnableIf<TIsTArray_V<ContainerType>, bool>::Type
 	TryGetField(const TSharedPtr<FJsonObject>& InJsonObject, const FString& InFieldName, ContainerType& OutValues);
 
 	// Map
@@ -411,7 +411,7 @@ namespace UE::Json
 		!std::is_same_v<ValueType, TSharedPtr<typename ValueType::ElementType>> &&
 		!std::is_same_v<ValueType, TJsonReference<typename ValueType::ElementType>> &&
 		!TIsTMap<ValueType>::Value &&
-		!TIsTArray<ValueType>::Value &&
+		!TIsTArray_V<ValueType> &&
 		!TypeTraits::TIsStringLike<ValueType>::Value &&
 		!TIsPODType<ValueType>::Value &&
 		!TypeTraits::THasFromJson<ValueType>::Value, bool>::Type
@@ -423,7 +423,7 @@ namespace UE::Json
 		!std::is_same_v<ValueType, TSharedPtr<typename ValueType::ElementType>> &&
 		!std::is_same_v<ValueType, TJsonReference<typename ValueType::ElementType>> &&
 		!TIsTMap<ValueType>::Value &&
-		!TIsTArray<ValueType>::Value &&
+		!TIsTArray_V<ValueType> &&
 		!TypeTraits::TIsStringLike<ValueType>::Value &&
 		!TIsPODType<ValueType>::Value &&
 		!TypeTraits::THasFromJson<ValueType>::Value, bool>::Type
diff --git a/Engine/Plugins/Experimental/Web/WebAPI/Source/WebAPIEditor/Public/WebAPIJsonUtilities.inl b/Engine/Plugins/Experimental/Web/WebAPI/Source/WebAPIEditor/Public/WebAPIJsonUtilities.inl
--- a/Engine/Plugins/Experimental/Web/WebAPI/Source/WebAPIEditor/Public/WebAPIJsonUtilities.inl
+++ b/Engine/Plugins/Experimental/Web/WebAPI/Source/WebAPIEditor/Public/WebAPIJsonUtilities.inl
@@ -1,4 +1,4 @@
-﻿// Copyright Epic Games, Inc. All Rights Reserved.
+// Copyright Epic Games, Inc. All Rights Reserved.
 
 #pragma once
 
@@ -156,7 +156,7 @@ namespace UE::Json
 
 	// Array
 	template <typename ContainerType>
-	typename TEnableIf<TIsTArray<ContainerType>::Value, bool>::Type
+	typename TEnableIf<TIsTArray_V<ContainerType>, bool>::Type
 	TryGet(const TSharedPtr<FJsonValue>& InJsonValue, ContainerType& OutValues)
 	{
 		using ValueType = typename ContainerType::ElementType;
@@ -179,7 +179,7 @@ namespace UE::Json
 	}
 
 	template <typename ContainerType>
-	constexpr typename TEnableIf<TIsTArray<ContainerType>::Value, bool>::Type
+	constexpr typename TEnableIf<TIsTArray_V<ContainerType>, bool>::Type
 	TryGetField(const TSharedPtr<FJsonObject>& InJsonObject, const FString& InFieldName, ContainerType& OutValues)
 	{
 		if(const TSharedPtr<FJsonValue> JsonField = InJsonObject->TryGetField(InFieldName))
@@ -336,7 +336,7 @@ namespace UE::Json
 		!std::is_same_v<ValueType, TSharedPtr<typename ValueType::ElementType>> &&
 		!std::is_same_v<ValueType, TJsonReference<typename ValueType::ElementType>> &&
 		!TIsTMap<ValueType>::Value &&
-		!TIsTArray<ValueType>::Value &&
+		!TIsTArray_V<ValueType> &&
 		!TypeTraits::TIsStringLike<ValueType>::Value &&
 		!TIsPODType<ValueType>::Value &&
 		!TypeTraits::THasFromJson<ValueType>::Value, bool>::Type
@@ -351,7 +351,7 @@ namespace UE::Json
 		!std::is_same_v<ValueType, TSharedPtr<typename ValueType::ElementType>> &&
 		!std::is_same_v<ValueType, TJsonReference<typename ValueType::ElementType>> &&
 		!TIsTMap<ValueType>::Value &&
-		!TIsTArray<ValueType>::Value &&
+		!TIsTArray_V<ValueType> &&
 		!TypeTraits::TIsStringLike<ValueType>::Value &&
 		!TIsPODType<ValueType>::Value &&
 		!TypeTraits::THasFromJson<ValueType>::Value, bool>::Type
diff --git a/Engine/Source/Runtime/Core/Public/Containers/Array.h b/Engine/Source/Runtime/Core/Public/Containers/Array.h
--- a/Engine/Source/Runtime/Core/Public/Containers/Array.h
+++ b/Engine/Source/Runtime/Core/Public/Containers/Array.h
@@ -3434,14 +3434,20 @@ struct TIsContiguousContainer<TArray<T, AllocatorType>>
 };
 
 /**
- * Traits class which determines whether or not a type is a TArray.
+ * Trait which determines whether or not a type is a TArray.
  */
-template <typename T> struct TIsTArray { enum { Value = false }; };
+template <typename T> constexpr bool TIsTArray_V = false;
 
-template <typename InElementType, typename InAllocatorType> struct TIsTArray<               TArray<InElementType, InAllocatorType>> { enum { Value = true }; };
-template <typename InElementType, typename InAllocatorType> struct TIsTArray<const          TArray<InElementType, InAllocatorType>> { enum { Value = true }; };
-template <typename InElementType, typename InAllocatorType> struct TIsTArray<      volatile TArray<InElementType, InAllocatorType>> { enum { Value = true }; };
-template <typename InElementType, typename InAllocatorType> struct TIsTArray<const volatile TArray<InElementType, InAllocatorType>> { enum { Value = true }; };
+template <typename InElementType, typename InAllocatorType> constexpr bool TIsTArray_V<               TArray<InElementType, InAllocatorType>> = true;
+template <typename InElementType, typename InAllocatorType> constexpr bool TIsTArray_V<const          TArray<InElementType, InAllocatorType>> = true;
+template <typename InElementType, typename InAllocatorType> constexpr bool TIsTArray_V<      volatile TArray<InElementType, InAllocatorType>> = true;
+template <typename InElementType, typename InAllocatorType> constexpr bool TIsTArray_V<const volatile TArray<InElementType, InAllocatorType>> = true;
+
+template <typename T>
+struct TIsTArray
+{
+	enum { Value = TIsTArray_V<T> };
+};
 
 
 //
diff --git a/Engine/Source/Runtime/Core/Public/Memory/CompositeBuffer.h b/Engine/Source/Runtime/Core/Public/Memory/CompositeBuffer.h
--- a/Engine/Source/Runtime/Core/Public/Memory/CompositeBuffer.h
+++ b/Engine/Source/Runtime/Core/Public/Memory/CompositeBuffer.h
@@ -32,7 +32,7 @@ class FCompositeBuffer
 	using CanAppendBufferType = std::bool_constant<
 		std::is_same_v<FCompositeBuffer, BufferType> ||
 		std::is_same_v<FSharedBuffer, BufferType> ||
-		TIsTArray<BufferType>::Value>;
+		TIsTArray_V<BufferType>>;
 
 public:
 	/**
diff --git a/Engine/Source/Runtime/CoreUObject/Public/UObject/ObjectPtr.h b/Engine/Source/Runtime/CoreUObject/Public/UObject/ObjectPtr.h
--- a/Engine/Source/Runtime/CoreUObject/Public/UObject/ObjectPtr.h
+++ b/Engine/Source/Runtime/CoreUObject/Public/UObject/ObjectPtr.h
@@ -723,7 +723,7 @@ PRAGMA_DISABLE_DEPRECATION_WARNINGS
 template <
 	typename ArrayType,
 	typename ArrayTypeNoRef = std::remove_reference_t<ArrayType>,
-	std::enable_if_t<TIsTArray<ArrayTypeNoRef>::Value>* = nullptr
+	std::enable_if_t<TIsTArray_V<ArrayTypeNoRef>>* = nullptr
 >
 decltype(auto) ToRawPtrTArrayUnsafe(ArrayType&& Array)
 {
diff --git a/Engine/Source/Runtime/RenderCore/Public/RenderGraphDefinitions.h b/Engine/Source/Runtime/RenderCore/Public/RenderGraphDefinitions.h
--- a/Engine/Source/Runtime/RenderCore/Public/RenderGraphDefinitions.h
+++ b/Engine/Source/Runtime/RenderCore/Public/RenderGraphDefinitions.h
@@ -757,6 +757,6 @@ using FRDGBufferInitialDataCallback = TFunction<const void*()>;
 using FRDGBufferInitialDataSizeCallback = TFunction<uint64()>;
 template <typename ArrayType, 
 	typename ArrayTypeNoRef = std::remove_reference_t<ArrayType>,
-	typename = typename TEnableIf<TIsTArray<ArrayTypeNoRef>::Value>::Type> using TRDGBufferArrayCallback = TFunction<const ArrayType&()>;
+	typename = typename TEnableIf<TIsTArray_V<ArrayTypeNoRef>>::Type> using TRDGBufferArrayCallback = TFunction<const ArrayType&()>;
 using FRDGBufferInitialDataFreeCallback = TFunction<void(const void* InData)>;
 using FRDGDispatchGroupCountCallback = TFunction<FIntVector()>;
diff --git a/Engine/Source/Runtime/RigVM/Public/RigVMCore/RigVMMemoryStorage.h b/Engine/Source/Runtime/RigVM/Public/RigVMCore/RigVMMemoryStorage.h
--- a/Engine/Source/Runtime/RigVM/Public/RigVMCore/RigVMMemoryStorage.h
+++ b/Engine/Source/Runtime/RigVM/Public/RigVMCore/RigVMMemoryStorage.h
@@ -366,7 +366,7 @@ public:
 	// returns true if the handle is of the given type
 	template <
 		typename T,
-		typename TEnableIf<TIsTArray<T>::Value>::Type* = nullptr
+		typename TEnableIf<TIsTArray_V<T>>::Type* = nullptr
 	>
 	bool IsType(bool bIsHiddenArgument = false) const
 	{
