CreateViewModels: check type at compile time instead of runtime

#rnx

diff --git a/Engine/Plugins/Experimental/Web/WebAPI/Source/WebAPIEditor/Private/Details/ViewModels/WebAPIViewModel.inl b/Engine/Plugins/Experimental/Web/WebAPI/Source/WebAPIEditor/Private/Details/ViewModels/WebAPIViewModel.inl
--- a/Engine/Plugins/Experimental/Web/WebAPI/Source/WebAPIEditor/Private/Details/ViewModels/WebAPIViewModel.inl
+++ b/Engine/Plugins/Experimental/Web/WebAPI/Source/WebAPIEditor/Private/Details/ViewModels/WebAPIViewModel.inl
@@ -18,38 +18,37 @@ TSharedPtr<ViewModelType> UE::WebAPI::Details::CreateViewModel(const TSharedRef<
 		return nullptr;
 	}
 	
-	const UClass* ModelClass = InModel->GetClass();
-	if(ModelClass == UWebAPIEnum::StaticClass())
+	if constexpr (std::is_same_v<UWebAPIEnum, ModelType>)
 	{
-		return FWebAPIEnumViewModel::Create(InParentViewModel, Cast<UWebAPIEnum>(InModel));
+		return FWebAPIEnumViewModel::Create(InParentViewModel, InModel);
 	}
-	else if(ModelClass == UWebAPIEnumValue::StaticClass())
+	else if constexpr (std::is_same_v<UWebAPIEnumValue, ModelType>)
+    {
+    	return FWebAPIEnumValueViewModel::Create(InParentViewModel, InModel);
+    }
+	else if constexpr (std::is_same_v<UWebAPIModel, ModelType>)
 	{
-		return FWebAPIEnumValueViewModel::Create(InParentViewModel, Cast<UWebAPIEnumValue>(InModel));
+		return FWebAPIModelViewModel::Create(InParentViewModel, InModel);
 	}
-	else if(ModelClass == UWebAPIModel::StaticClass())
+	else if constexpr (std::is_same_v<UWebAPIProperty, ModelType>)
 	{
-		return FWebAPIModelViewModel::Create(InParentViewModel, Cast<UWebAPIModel>(InModel));
+		return FWebAPIPropertyViewModel::Create(InParentViewModel, InModel);
 	}
-	else if(ModelClass == UWebAPIProperty::StaticClass())
+	else if constexpr (std::is_same_v<UWebAPIService, ModelType>)
 	{
-		return FWebAPIPropertyViewModel::Create(InParentViewModel, Cast<UWebAPIProperty>(InModel));
+		return FWebAPIServiceViewModel::Create(InParentViewModel, InModel);
 	}
-	else if(ModelClass == UWebAPIService::StaticClass())
+	else if constexpr (std::is_same_v<UWebAPIOperation, ModelType>)
 	{
-		return FWebAPIServiceViewModel::Create(InParentViewModel, Cast<UWebAPIService>(InModel));
+		return FWebAPIOperationViewModel::Create(InParentViewModel, InModel);
 	}
-	else if(ModelClass == UWebAPIOperation::StaticClass())
+	else if constexpr (std::is_same_v<UWebAPIParameter, ModelType>)
 	{
-		return FWebAPIOperationViewModel::Create(InParentViewModel, Cast<UWebAPIOperation>(InModel));
-	}
-	else if(ModelClass == UWebAPIParameter::StaticClass())
-	{
-		return FWebAPIParameterViewModel::Create(InParentViewModel, Cast<UWebAPIParameter>(InModel));
+		return FWebAPIParameterViewModel::Create(InParentViewModel, InModel);
 	}
 	else
 	{
-		checkNoEntry();
+		static_assert(false, "Unsupported model type");
 		return nullptr;
 	}
 }
