[Backout] - CL27338179
Original CL Desc
-----------------------------------------------------------------
CreateViewModels: check type at compile time instead of runtime

#rnx

diff --git a/Engine/Plugins/Experimental/Web/WebAPI/Source/WebAPIEditor/Private/Details/ViewModels/WebAPIViewModel.inl b/Engine/Plugins/Experimental/Web/WebAPI/Source/WebAPIEditor/Private/Details/ViewModels/WebAPIViewModel.inl
--- a/Engine/Plugins/Experimental/Web/WebAPI/Source/WebAPIEditor/Private/Details/ViewModels/WebAPIViewModel.inl
+++ b/Engine/Plugins/Experimental/Web/WebAPI/Source/WebAPIEditor/Private/Details/ViewModels/WebAPIViewModel.inl
@@ -18,37 +18,38 @@ TSharedPtr<ViewModelType> UE::WebAPI::Details::CreateViewModel(const TSharedRef<
 		return nullptr;
 	}
 	
-	if constexpr (std::is_same_v<UWebAPIEnum, ModelType>)
+	const UClass* ModelClass = InModel->GetClass();
+	if(ModelClass == UWebAPIEnum::StaticClass())
 	{
-		return FWebAPIEnumViewModel::Create(InParentViewModel, InModel);
+		return FWebAPIEnumViewModel::Create(InParentViewModel, Cast<UWebAPIEnum>(InModel));
 	}
-	else if constexpr (std::is_same_v<UWebAPIEnumValue, ModelType>)
-    {
-    	return FWebAPIEnumValueViewModel::Create(InParentViewModel, InModel);
-    }
-	else if constexpr (std::is_same_v<UWebAPIModel, ModelType>)
+	else if(ModelClass == UWebAPIEnumValue::StaticClass())
 	{
-		return FWebAPIModelViewModel::Create(InParentViewModel, InModel);
+		return FWebAPIEnumValueViewModel::Create(InParentViewModel, Cast<UWebAPIEnumValue>(InModel));
 	}
-	else if constexpr (std::is_same_v<UWebAPIProperty, ModelType>)
+	else if(ModelClass == UWebAPIModel::StaticClass())
 	{
-		return FWebAPIPropertyViewModel::Create(InParentViewModel, InModel);
+		return FWebAPIModelViewModel::Create(InParentViewModel, Cast<UWebAPIModel>(InModel));
 	}
-	else if constexpr (std::is_same_v<UWebAPIService, ModelType>)
+	else if(ModelClass == UWebAPIProperty::StaticClass())
 	{
-		return FWebAPIServiceViewModel::Create(InParentViewModel, InModel);
+		return FWebAPIPropertyViewModel::Create(InParentViewModel, Cast<UWebAPIProperty>(InModel));
 	}
-	else if constexpr (std::is_same_v<UWebAPIOperation, ModelType>)
+	else if(ModelClass == UWebAPIService::StaticClass())
 	{
-		return FWebAPIOperationViewModel::Create(InParentViewModel, InModel);
+		return FWebAPIServiceViewModel::Create(InParentViewModel, Cast<UWebAPIService>(InModel));
 	}
-	else if constexpr (std::is_same_v<UWebAPIParameter, ModelType>)
+	else if(ModelClass == UWebAPIOperation::StaticClass())
 	{
-		return FWebAPIParameterViewModel::Create(InParentViewModel, InModel);
+		return FWebAPIOperationViewModel::Create(InParentViewModel, Cast<UWebAPIOperation>(InModel));
+	}
+	else if(ModelClass == UWebAPIParameter::StaticClass())
+	{
+		return FWebAPIParameterViewModel::Create(InParentViewModel, Cast<UWebAPIParameter>(InModel));
 	}
 	else
 	{
-		static_assert(false, "Unsupported model type");
+		checkNoEntry();
 		return nullptr;
 	}
 }
