CreateViewModels: check type at compile time instead of runtime

#rnx

diff --git a/Engine/Plugins/Experimental/Web/WebAPI/Source/WebAPIEditor/Private/Details/ViewModels/WebAPIViewModel.inl b/Engine/Plugins/Experimental/Web/WebAPI/Source/WebAPIEditor/Private/Details/ViewModels/WebAPIViewModel.inl
--- a/Engine/Plugins/Experimental/Web/WebAPI/Source/WebAPIEditor/Private/Details/ViewModels/WebAPIViewModel.inl
+++ b/Engine/Plugins/Experimental/Web/WebAPI/Source/WebAPIEditor/Private/Details/ViewModels/WebAPIViewModel.inl
@@ -17,39 +17,55 @@ TSharedPtr<ViewModelType> UE::WebAPI::Details::CreateViewModel(const TSharedRef<
 	{
 		return nullptr;
 	}
-	
+
 	const UClass* ModelClass = InModel->GetClass();
-	if(ModelClass == UWebAPIEnum::StaticClass())
-	{
-		return FWebAPIEnumViewModel::Create(InParentViewModel, Cast<UWebAPIEnum>(InModel));
-	}
-	else if(ModelClass == UWebAPIEnumValue::StaticClass())
-	{
-		return FWebAPIEnumValueViewModel::Create(InParentViewModel, Cast<UWebAPIEnumValue>(InModel));
-	}
-	else if(ModelClass == UWebAPIModel::StaticClass())
-	{
-		return FWebAPIModelViewModel::Create(InParentViewModel, Cast<UWebAPIModel>(InModel));
-	}
-	else if(ModelClass == UWebAPIProperty::StaticClass())
-	{
-		return FWebAPIPropertyViewModel::Create(InParentViewModel, Cast<UWebAPIProperty>(InModel));
-	}
-	else if(ModelClass == UWebAPIService::StaticClass())
-	{
-		return FWebAPIServiceViewModel::Create(InParentViewModel, Cast<UWebAPIService>(InModel));
-	}
-	else if(ModelClass == UWebAPIOperation::StaticClass())
+	if constexpr (std::is_base_of_v<UWebAPIModelBase, ModelType>)
 	{
-		return FWebAPIOperationViewModel::Create(InParentViewModel, Cast<UWebAPIOperation>(InModel));
+		if(ModelClass == UWebAPIEnum::StaticClass())
+		{
+			return FWebAPIEnumViewModel::Create(InParentViewModel, Cast<UWebAPIEnum>(InModel));
+		}
+		else if(ModelClass == UWebAPIEnumValue::StaticClass())
+		{
+			return FWebAPIEnumValueViewModel::Create(InParentViewModel, Cast<UWebAPIEnumValue>(InModel));
+		}
+		else if(ModelClass == UWebAPIModel::StaticClass())
+		{
+			return FWebAPIModelViewModel::Create(InParentViewModel, Cast<UWebAPIModel>(InModel));
+		}
+		else if(ModelClass == UWebAPIProperty::StaticClass())
+		{
+			return FWebAPIPropertyViewModel::Create(InParentViewModel, Cast<UWebAPIProperty>(InModel));
+		}
+		else if(ModelClass == UWebAPIService::StaticClass())
+		{
+			return FWebAPIServiceViewModel::Create(InParentViewModel, Cast<UWebAPIService>(InModel));
+		}
+		else if(ModelClass == UWebAPIParameter::StaticClass())
+		{
+			return FWebAPIParameterViewModel::Create(InParentViewModel, Cast<UWebAPIParameter>(InModel));
+		}
+		else
+		{
+			checkNoEntry();
+			return nullptr;
+		}
 	}
-	else if(ModelClass == UWebAPIParameter::StaticClass())
+	else if constexpr (std::is_base_of_v<UWebAPIOperation, ModelType>)
 	{
-		return FWebAPIParameterViewModel::Create(InParentViewModel, Cast<UWebAPIParameter>(InModel));
+		if(ModelClass == UWebAPIOperation::StaticClass())
+		{
+			return FWebAPIOperationViewModel::Create(InParentViewModel, Cast<UWebAPIOperation>(InModel));
+		}
+		else
+		{
+			checkNoEntry();
+			return nullptr;
+		}
 	}
 	else
 	{
-		checkNoEntry();
+		static_assert(sizeof(ModelType) == 0, "Unsupported type");
 		return nullptr;
 	}
 }
