WebAPIOpenAPI3.spec.cpp : check if cast can be executed
#rnx

diff --git a/Engine/Plugins/Experimental/Web/WebAPI/Source/WebAPIOpenAPI/Private/Tests/WebAPIOpenAPI3.spec.cpp b/Engine/Plugins/Experimental/Web/WebAPI/Source/WebAPIOpenAPI/Private/Tests/WebAPIOpenAPI3.spec.cpp
--- a/Engine/Plugins/Experimental/Web/WebAPI/Source/WebAPIOpenAPI/Private/Tests/WebAPIOpenAPI3.spec.cpp
+++ b/Engine/Plugins/Experimental/Web/WebAPI/Source/WebAPIOpenAPI/Private/Tests/WebAPIOpenAPI3.spec.cpp
@@ -171,17 +171,26 @@ BEGIN_DEFINE_SPEC(FWebAPIOpenAPI3Spec,
 		const TObjectPtr<InputModelType>* FoundModel = InModels.FindByPredicate([InName](const TObjectPtr<InputModelType> InModelBase)
 		{
 			// @note: order matters!
-			if(const TObjectPtr<UWebAPIParameter> Parameter = Cast<UWebAPIParameter>(InModelBase))
+			if constexpr (std::is_base_of_v<UWebAPIParameter, InputModelType>)
 			{
-				return Parameter->Name.ToString(true).Equals(InName);
+				if(const TObjectPtr<UWebAPIParameter> Parameter = Cast<UWebAPIParameter>(InModelBase))
+				{
+					return Parameter->Name.ToString(true).Equals(InName);
+				}
 			}
-			else if(const TObjectPtr<UWebAPIModel> Model = Cast<UWebAPIModel>(InModelBase))
+			if constexpr (std::is_base_of_v<UWebAPIModel, InputModelType>)
 			{
-				return Model->Name.ToString(true).Equals(InName);
+				if(const TObjectPtr<UWebAPIModel> Model = Cast<UWebAPIModel>(InModelBase))
+				{
+					return Model->Name.ToString(true).Equals(InName);
+				}
 			}
-			else if(const TObjectPtr<UWebAPIEnum> Enum = Cast<UWebAPIEnum>(InModelBase))
+			if constexpr (std::is_base_of_v<UWebAPIEnum, InputModelType>)
 			{
-				return Enum->Name.ToString(true).Equals(InName);
+				if(const TObjectPtr<UWebAPIEnum> Enum = Cast<UWebAPIEnum>(InModelBase))
+				{
+					return Enum->Name.ToString(true).Equals(InName);
+				}
 			}
 
 			return false;
