[WebAPI] Removed erroneous raw ptr cleanup

diff --git a/Engine/Plugins/Experimental/Web/WebAPI/Source/WebAPIBlueprintGraph/Private/WebAPIBlueprintGraphUtilities.cpp b/Engine/Plugins/Experimental/Web/WebAPI/Source/WebAPIBlueprintGraph/Private/WebAPIBlueprintGraphUtilities.cpp
--- a/Engine/Plugins/Experimental/Web/WebAPI/Source/WebAPIBlueprintGraph/Private/WebAPIBlueprintGraphUtilities.cpp
+++ b/Engine/Plugins/Experimental/Web/WebAPI/Source/WebAPIBlueprintGraph/Private/WebAPIBlueprintGraphUtilities.cpp
@@ -24,16 +24,11 @@ namespace UE::WebAPI
 				//  If valid, just return it
 				if(Found->Get<Index>().IsValid())
 				{
-					delete CachedOutcome;
-					CachedOutcome = nullptr;
-					
 					return Found->Get<Index>().Get();
 				}
-				else
-				{
-					// Set the pointer for later use
-					CachedOutcome = Found;
-				}
+				
+				// Set the pointer for later use
+				CachedOutcome = Found;
 			}
 			// Otherwise add it as a cache entry
 			else
@@ -41,21 +36,18 @@ namespace UE::WebAPI
 				CachedOutcome = &CachedOutcomeDelegates.Emplace(InOperationClass->GetFName());
 			}
 
-			int32 C = 0;
 			FMulticastDelegateProperty* P = nullptr;
 			const FString OutcomeNameStr = InOutcomeName.ToString();
 			for (TFieldIterator<FMulticastDelegateProperty> PropertyIterator(InOperationClass); PropertyIterator; ++PropertyIterator)
 			{
 				if(PropertyIterator->GetName().Contains(OutcomeNameStr))
 				{
-					++C;
 					P = *PropertyIterator;
 					CachedOutcome->Set<Index>(*PropertyIterator);
 					return P;
 				}
 			}
 
-			CachedOutcome = nullptr;
 			return nullptr;
 		}
 		
