Added deprecation guards and comment to highlight an issue that will cause a deprecation warning with new TSoftObjectPtr changes.

diff --git a/Engine/Plugins/Experimental/Web/WebAPI/Source/WebAPIEditor/Private/Dom/WebAPIOperation.cpp b/Engine/Plugins/Experimental/Web/WebAPI/Source/WebAPIEditor/Private/Dom/WebAPIOperation.cpp
--- a/Engine/Plugins/Experimental/Web/WebAPI/Source/WebAPIEditor/Private/Dom/WebAPIOperation.cpp
+++ b/Engine/Plugins/Experimental/Web/WebAPI/Source/WebAPIEditor/Private/Dom/WebAPIOperation.cpp
@@ -100,16 +100,23 @@ void UWebAPIOperation::Visit(TFunctionRef<void(IWebAPISchemaObjectInterface*&)>
 	}
 }
 
+PRAGMA_DISABLE_DEPRECATION_WARNINGS
+
 void UWebAPIOperation::BindToTypeInfo()
 {
 	check(Name.HasTypeInfo());
 
 	if(Name.TypeInfo->Model.IsNull())
 	{
+		// This has been assigning a UWebAPIOperation* to a TSoftObjectPtr<UWebAPIModelBase> since it was submitted,
+		// even though those tpyes are unrelated.  This is dangerous if someone uses the pointer, but work if only
+		// the path is used.
 		Name.TypeInfo->Model = this;
 	}
 }
 
+PRAGMA_ENABLE_DEPRECATION_WARNINGS
+
 #if WITH_EDITOR
 void UWebAPIOperation::SetCodeText(const FString& InCodeText)
 {
