[WebAPI] fix

#rnx

diff --git a/Engine/Plugins/Experimental/Web/WebAPI/Source/WebAPIEditor/Private/Details/ViewModels/WebAPIModelViewModel.cpp b/Engine/Plugins/Experimental/Web/WebAPI/Source/WebAPIEditor/Private/Details/ViewModels/WebAPIModelViewModel.cpp
--- a/Engine/Plugins/Experimental/Web/WebAPI/Source/WebAPIEditor/Private/Details/ViewModels/WebAPIModelViewModel.cpp
+++ b/Engine/Plugins/Experimental/Web/WebAPI/Source/WebAPIEditor/Private/Details/ViewModels/WebAPIModelViewModel.cpp
@@ -45,15 +45,18 @@ bool FWebAPIPropertyViewModel::HasCodeText() const
 		return IWebAPIViewModel::HasCodeText();
 	}
 
-	if(Property->Type.HasTypeInfo() && !Property->Type.TypeInfo->Model.IsNull())
+	if(Property->Type.HasTypeInfo())
 	{
-		if(const TObjectPtr<UWebAPIModel> AsModel = Cast<UWebAPIModel>(Property->Type.TypeInfo->Model.Get()))
+		if(const UObject* PropertyModel = Property->Type.TypeInfo->GetModel())
 		{
-			return !AsModel->GeneratedCodeText.IsEmpty();
-		}
-		else if(const TObjectPtr<UWebAPIEnum> AsEnum = Cast<UWebAPIEnum>(Property->Type.TypeInfo->Model.Get()))
-		{
-			return !AsEnum->GeneratedCodeText.IsEmpty();
+			if(const UWebAPIModel* AsModel = Cast<UWebAPIModel>(PropertyModel))
+			{
+				return !AsModel->GeneratedCodeText.IsEmpty();
+			}
+			else if(const UWebAPIEnum* AsEnum = Cast<UWebAPIEnum>(PropertyModel))
+			{
+				return !AsEnum->GeneratedCodeText.IsEmpty();
+			}
 		}
 	}
 
@@ -67,15 +70,18 @@ FText FWebAPIPropertyViewModel::GetCodeText() const
 		return IWebAPIViewModel::GetCodeText();
 	}
 
-	if(Property->Type.HasTypeInfo() && !Property->Type.TypeInfo->Model.IsNull())
+	if(Property->Type.HasTypeInfo())
 	{
-		if(const TObjectPtr<UWebAPIModel> AsModel = Cast<UWebAPIModel>(Property->Type.TypeInfo->Model.Get()))
-		{
-			return FText::FromString(AsModel->GeneratedCodeText);
-		}
-		else if(const TObjectPtr<UWebAPIEnum> AsEnum = Cast<UWebAPIEnum>(Property->Type.TypeInfo->Model.Get()))
+		if(const UObject* PropertyModel = Property->Type.TypeInfo->GetModel())
 		{
-			return FText::FromString(AsEnum->GeneratedCodeText);
+			if(const UWebAPIModel* AsModel = Cast<UWebAPIModel>(PropertyModel))
+			{
+				return FText::FromString(AsModel->GeneratedCodeText);
+			}
+			else if(const UWebAPIEnum* AsEnum = Cast<UWebAPIEnum>(PropertyModel))
+			{
+				return FText::FromString(AsEnum->GeneratedCodeText);
+			}
 		}
 	}
  
@@ -114,9 +120,12 @@ void FWebAPIPropertyViewModel::Initialize()
 	const FString Description = Property->Description;
 	const FWebAPITypeNameVariant Type = Property->Type;
 
-	if(Type.HasTypeInfo() && Type.TypeInfo->bIsNested && !Type.TypeInfo->Model.IsNull())
+	if(Type.HasTypeInfo() && Type.TypeInfo->bIsNested)
 	{
-		NestedModel = UE::WebAPI::Details::CreateViewModel(AsShared(), Type.TypeInfo->Model.Get());
+		if(UObject* PropertyModel = Type.TypeInfo->GetModel())
+		{
+			NestedModel = UE::WebAPI::Details::CreateViewModel(SharedThis(this), PropertyModel);
+		}
 	}
 
 	FString DebugLabel;
@@ -124,9 +133,9 @@ void FWebAPIPropertyViewModel::Initialize()
 	bool bIsTypeResolved = true;
 	if(Type.HasTypeInfo() && !Type.TypeInfo->bIsBuiltinType)
 	{
-		bIsTypeResolved = !Type.TypeInfo->Model.IsNull();
+		bIsTypeResolved = Type.TypeInfo->GetModel() != nullptr;
 	}
-	
+
 	if(!bIsTypeResolved)
 	{
 		DebugLabel = TEXT(" (Unresolved)");
diff --git a/Engine/Plugins/Experimental/Web/WebAPI/Source/WebAPIEditor/Private/Details/ViewModels/WebAPIOperationParameterViewModel.cpp b/Engine/Plugins/Experimental/Web/WebAPI/Source/WebAPIEditor/Private/Details/ViewModels/WebAPIOperationParameterViewModel.cpp
--- a/Engine/Plugins/Experimental/Web/WebAPI/Source/WebAPIEditor/Private/Details/ViewModels/WebAPIOperationParameterViewModel.cpp
+++ b/Engine/Plugins/Experimental/Web/WebAPI/Source/WebAPIEditor/Private/Details/ViewModels/WebAPIOperationParameterViewModel.cpp
@@ -59,7 +59,7 @@ void FWebAPIOperationParameterViewModel::Initialize()
 	bool bIsTypeResolved = true;
 	if(Type.HasTypeInfo() && !Type.TypeInfo->bIsBuiltinType)
 	{
-		bIsTypeResolved = !Type.TypeInfo->Model.IsNull();
+		bIsTypeResolved = Type.TypeInfo->GetModel() != nullptr;
 	}
 	
 	if(!bIsTypeResolved)
diff --git a/Engine/Plugins/Experimental/Web/WebAPI/Source/WebAPIEditor/Private/Details/ViewModels/WebAPIParameterViewModel.cpp b/Engine/Plugins/Experimental/Web/WebAPI/Source/WebAPIEditor/Private/Details/ViewModels/WebAPIParameterViewModel.cpp
--- a/Engine/Plugins/Experimental/Web/WebAPI/Source/WebAPIEditor/Private/Details/ViewModels/WebAPIParameterViewModel.cpp
+++ b/Engine/Plugins/Experimental/Web/WebAPI/Source/WebAPIEditor/Private/Details/ViewModels/WebAPIParameterViewModel.cpp
@@ -87,9 +87,9 @@ void FWebAPIParameterViewModel::Initialize()
 
 	if(Type.HasTypeInfo() && Type.TypeInfo->bIsNested)
 	{
-		NestedModel = UE::WebAPI::Details::CreateViewModel(AsShared(), Type.TypeInfo->Model.Get());
+		NestedModel = UE::WebAPI::Details::CreateViewModel<FWebAPIParameterViewModel, IWebAPIViewModel>(SharedThis(this), Type.TypeInfo->GetModel());
 	}
-	
+
 	CachedLabel = FText::FromString(Parameter->Name.GetDisplayName());
 
 	TArray<FText> TooltipArgs;
diff --git a/Engine/Plugins/Experimental/Web/WebAPI/Source/WebAPIEditor/Private/Details/ViewModels/WebAPIViewModel.h b/Engine/Plugins/Experimental/Web/WebAPI/Source/WebAPIEditor/Private/Details/ViewModels/WebAPIViewModel.h
--- a/Engine/Plugins/Experimental/Web/WebAPI/Source/WebAPIEditor/Private/Details/ViewModels/WebAPIViewModel.h
+++ b/Engine/Plugins/Experimental/Web/WebAPI/Source/WebAPIEditor/Private/Details/ViewModels/WebAPIViewModel.h
@@ -7,10 +7,11 @@
 class UWebAPIDefinition;
 class UWebAPISchema;
 
-class FWebAPIModelViewModel;
 class FWebAPIDefinitionViewModel;
+class FWebAPIModelViewModel;
 class FWebAPISchemaViewModel;
 class FWebAPIServiceViewModel;
+class IWebAPISchemaObjectInterface;
 class IWebAPIViewModel;
 
 namespace UE
@@ -19,8 +20,8 @@ namespace UE
 	{
 		namespace Details
 		{
-			template <typename ModelType, class ParentViewModelType = IWebAPIViewModel, class ViewModelType = IWebAPIViewModel> 
-			TSharedPtr<ViewModelType> CreateViewModel(const TSharedRef<ParentViewModelType>& InParentViewModel, ModelType* InModel);
+			template <class ParentViewModelType = IWebAPIViewModel, class ViewModelType = IWebAPIViewModel>
+			TSharedPtr<ViewModelType> CreateViewModel(const TSharedRef<ParentViewModelType>& InParentViewModel, UObject* InModel);
 		}
 	}
 }
diff --git a/Engine/Plugins/Experimental/Web/WebAPI/Source/WebAPIEditor/Private/Details/ViewModels/WebAPIViewModel.inl b/Engine/Plugins/Experimental/Web/WebAPI/Source/WebAPIEditor/Private/Details/ViewModels/WebAPIViewModel.inl
--- a/Engine/Plugins/Experimental/Web/WebAPI/Source/WebAPIEditor/Private/Details/ViewModels/WebAPIViewModel.inl
+++ b/Engine/Plugins/Experimental/Web/WebAPI/Source/WebAPIEditor/Private/Details/ViewModels/WebAPIViewModel.inl
@@ -2,70 +2,74 @@
 
 #pragma once
 
+#include "Dom/WebAPIService.h"
+#include "WebAPIEditorLog.h"
 #include "WebAPIEnumViewModel.h"
 #include "WebAPIModelViewModel.h"
 #include "WebAPIOperationViewModel.h"
 #include "WebAPIParameterViewModel.h"
 #include "WebAPIServiceViewModel.h"
 #include "WebAPIViewModel.h"
-#include "Dom/WebAPIService.h"
 
-template <typename ModelType, class ParentViewModelType, class ViewModelType>
-TSharedPtr<ViewModelType> UE::WebAPI::Details::CreateViewModel(const TSharedRef<ParentViewModelType>& InParentViewModel, ModelType* InModel)
+namespace UE::WebAPI::Details
 {
- 	if(!InModel)
-	{
-		return nullptr;
-	}
-
-	const UClass* ModelClass = InModel->GetClass();
-	if constexpr (std::is_base_of_v<UWebAPIModelBase, ModelType>)
+	template <class ParentViewModelType, class ViewModelType>
+	TSharedPtr<ViewModelType> CreateViewModel(const TSharedRef<ParentViewModelType>& InParentViewModel, UObject* InModel)
 	{
-		if(ModelClass == UWebAPIEnum::StaticClass())
-		{
-			return FWebAPIEnumViewModel::Create(InParentViewModel, Cast<UWebAPIEnum>(InModel));
-		}
-		else if(ModelClass == UWebAPIEnumValue::StaticClass())
-		{
-			return FWebAPIEnumValueViewModel::Create(InParentViewModel, Cast<UWebAPIEnumValue>(InModel));
-		}
-		else if(ModelClass == UWebAPIModel::StaticClass())
+		if(!InModel)
 		{
-			return FWebAPIModelViewModel::Create(InParentViewModel, Cast<UWebAPIModel>(InModel));
-		}
-		else if(ModelClass == UWebAPIProperty::StaticClass())
-		{
-			return FWebAPIPropertyViewModel::Create(InParentViewModel, Cast<UWebAPIProperty>(InModel));
-		}
-		else if(ModelClass == UWebAPIService::StaticClass())
-		{
-			return FWebAPIServiceViewModel::Create(InParentViewModel, Cast<UWebAPIService>(InModel));
+			return nullptr;
 		}
-		else if(ModelClass == UWebAPIParameter::StaticClass())
+
+		const UClass* ModelClass = InModel->GetClass();
+		if(ModelClass->IsChildOf<UWebAPIModelBase>())
 		{
-			return FWebAPIParameterViewModel::Create(InParentViewModel, Cast<UWebAPIParameter>(InModel));
+			if(ModelClass == UWebAPIEnum::StaticClass())
+			{
+				return FWebAPIEnumViewModel::Create(InParentViewModel, Cast<UWebAPIEnum>(InModel));
+			}
+			else if(ModelClass == UWebAPIEnumValue::StaticClass())
+			{
+				return FWebAPIEnumValueViewModel::Create(InParentViewModel, Cast<UWebAPIEnumValue>(InModel));
+			}
+			else if(ModelClass == UWebAPIModel::StaticClass())
+			{
+				return FWebAPIModelViewModel::Create(InParentViewModel, Cast<UWebAPIModel>(InModel));
+			}
+			else if(ModelClass == UWebAPIProperty::StaticClass())
+			{
+				return FWebAPIPropertyViewModel::Create(InParentViewModel, Cast<UWebAPIProperty>(InModel));
+			}
+			else if(ModelClass == UWebAPIService::StaticClass())
+			{
+				return FWebAPIServiceViewModel::Create(InParentViewModel, Cast<UWebAPIService>(InModel));
+			}
+			else if(ModelClass == UWebAPIParameter::StaticClass())
+			{
+				return FWebAPIParameterViewModel::Create(InParentViewModel, Cast<UWebAPIParameter>(InModel));
+			}
+			else
+			{
+				checkNoEntry();
+				return nullptr;
+			}
 		}
-		else
+		else if(ModelClass->IsChildOf<UWebAPIOperation>())
 		{
-			checkNoEntry();
-			return nullptr;
-		}
-	}
-	else if constexpr (std::is_base_of_v<UWebAPIOperation, ModelType>)
-	{
-		if(ModelClass == UWebAPIOperation::StaticClass())
-		{
-			return FWebAPIOperationViewModel::Create(InParentViewModel, Cast<UWebAPIOperation>(InModel));
+			if(ModelClass == UWebAPIOperation::StaticClass())
+			{
+				return FWebAPIOperationViewModel::Create(InParentViewModel, Cast<UWebAPIOperation>(InModel));
+			}
+			else
+			{
+				checkNoEntry();
+				return nullptr;
+			}
 		}
 		else
 		{
-			checkNoEntry();
+			UE_LOG(LogWebAPIEditor, Error, TEXT("Unsupported Type"));
 			return nullptr;
 		}
 	}
-	else
-	{
-		static_assert(sizeof(ModelType) == 0, "Unsupported type");
-		return nullptr;
-	}
 }
diff --git a/Engine/Plugins/Experimental/Web/WebAPI/Source/WebAPIEditor/Private/Dom/WebAPIEnum.cpp b/Engine/Plugins/Experimental/Web/WebAPI/Source/WebAPIEditor/Private/Dom/WebAPIEnum.cpp
--- a/Engine/Plugins/Experimental/Web/WebAPI/Source/WebAPIEditor/Private/Dom/WebAPIEnum.cpp
+++ b/Engine/Plugins/Experimental/Web/WebAPI/Source/WebAPIEditor/Private/Dom/WebAPIEnum.cpp
@@ -62,9 +62,9 @@ void UWebAPIEnum::BindToTypeInfo()
 
 	check(Name.HasTypeInfo());
 
-	if(!Name.TypeInfo->bIsBuiltinType && Name.TypeInfo->Model.IsNull())
+	if(!Name.TypeInfo->bIsBuiltinType && Name.TypeInfo->GetModel() == nullptr)
 	{
-		Name.TypeInfo->Model = this;
+		Name.TypeInfo->SetModel(this);
 		if(Name.TypeInfo->DefaultValue.IsEmpty())
 		{
 			Name.TypeInfo->DefaultValue = GetDefaultValue();
diff --git a/Engine/Plugins/Experimental/Web/WebAPI/Source/WebAPIEditor/Private/Dom/WebAPIModel.cpp b/Engine/Plugins/Experimental/Web/WebAPI/Source/WebAPIEditor/Private/Dom/WebAPIModel.cpp
--- a/Engine/Plugins/Experimental/Web/WebAPI/Source/WebAPIEditor/Private/Dom/WebAPIModel.cpp
+++ b/Engine/Plugins/Experimental/Web/WebAPI/Source/WebAPIEditor/Private/Dom/WebAPIModel.cpp
@@ -41,11 +41,14 @@ FString UWebAPIProperty::GetDefaultValue(bool bQualified) const
 	UWebAPIDefinition* OwningDefinition = GetTypedOuter<UWebAPIDefinition>();
 	check(OwningDefinition);
 	
-	if(Type.HasTypeInfo() && !Type.TypeInfo->Model.IsNull())
+	if(Type.HasTypeInfo())
 	{
-		if(const TObjectPtr<UWebAPIEnum> Enum = Cast<UWebAPIEnum>(Type.TypeInfo->Model.Get()))
+		if(const UObject* Model = Type.TypeInfo->GetModel())
 		{
-			return Enum->GetDefaultValue(bQualified);
+			if(const UWebAPIEnum* Enum = Cast<UWebAPIEnum>(Model))
+			{
+				return Enum->GetDefaultValue(bQualified);
+			}
 		}
 	}
 
@@ -116,9 +119,9 @@ void UWebAPIModel::BindToTypeInfo()
 
 	check(Name.HasTypeInfo());
 
-	if(!Name.TypeInfo->bIsBuiltinType && Name.TypeInfo->Model.IsNull())
+	if(!Name.TypeInfo->bIsBuiltinType && Name.TypeInfo->GetModel() == nullptr)
 	{
-		Name.TypeInfo->Model = this;
+		Name.TypeInfo->SetModel(this);
 	}
 }
 
diff --git a/Engine/Plugins/Experimental/Web/WebAPI/Source/WebAPIEditor/Private/Dom/WebAPIOperation.cpp b/Engine/Plugins/Experimental/Web/WebAPI/Source/WebAPIEditor/Private/Dom/WebAPIOperation.cpp
--- a/Engine/Plugins/Experimental/Web/WebAPI/Source/WebAPIEditor/Private/Dom/WebAPIOperation.cpp
+++ b/Engine/Plugins/Experimental/Web/WebAPI/Source/WebAPIEditor/Private/Dom/WebAPIOperation.cpp
@@ -106,12 +106,9 @@ void UWebAPIOperation::BindToTypeInfo()
 {
 	check(Name.HasTypeInfo());
 
-	if(Name.TypeInfo->Model.IsNull())
+	if(!Name.TypeInfo->GetModel())
 	{
-		// This has been assigning a UWebAPIOperation* to a TSoftObjectPtr<UWebAPIModelBase> since it was submitted,
-		// even though those tpyes are unrelated.  This is dangerous if someone uses the pointer, but work if only
-		// the path is used.
-		Name.TypeInfo->Model = this;
+		Name.TypeInfo->SetModel(this);
 	}
 }
 
diff --git a/Engine/Plugins/Experimental/Web/WebAPI/Source/WebAPIEditor/Private/Dom/WebAPISchema.cpp b/Engine/Plugins/Experimental/Web/WebAPI/Source/WebAPIEditor/Private/Dom/WebAPISchema.cpp
--- a/Engine/Plugins/Experimental/Web/WebAPI/Source/WebAPIEditor/Private/Dom/WebAPISchema.cpp
+++ b/Engine/Plugins/Experimental/Web/WebAPI/Source/WebAPIEditor/Private/Dom/WebAPISchema.cpp
@@ -242,7 +242,7 @@ UWebAPISchema::UWebAPISchema()
 	TypeRegistry = CreateDefaultSubobject<UWebAPITypeRegistry>(TEXT("TypeRegistry"));
 }
 
-TObjectPtr<UWebAPIEnum> UWebAPISchema::AddEnum(const TObjectPtr<UWebAPITypeInfo>& InTypeInfo)
+UWebAPIEnum* UWebAPISchema::AddEnum(const TObjectPtr<UWebAPITypeInfo>& InTypeInfo)
 {
 	check(InTypeInfo);
 
@@ -280,7 +280,7 @@ TObjectPtr<UWebAPIEnum> UWebAPISchema::AddEnum(const TObjectPtr<UWebAPITypeInfo>
 	return Cast<UWebAPIEnum>(AddedModel);
 }
 
-TObjectPtr<UWebAPIParameter> UWebAPISchema::AddParameter(const TObjectPtr<UWebAPITypeInfo>& InTypeInfo)
+UWebAPIParameter* UWebAPISchema::AddParameter(const TObjectPtr<UWebAPITypeInfo>& InTypeInfo)
 {
 	check(InTypeInfo);
 
diff --git a/Engine/Plugins/Experimental/Web/WebAPI/Source/WebAPIEditor/Private/Dom/WebAPIType.cpp b/Engine/Plugins/Experimental/Web/WebAPI/Source/WebAPIEditor/Private/Dom/WebAPIType.cpp
--- a/Engine/Plugins/Experimental/Web/WebAPI/Source/WebAPIEditor/Private/Dom/WebAPIType.cpp
+++ b/Engine/Plugins/Experimental/Web/WebAPI/Source/WebAPIEditor/Private/Dom/WebAPIType.cpp
@@ -118,6 +118,16 @@ void UWebAPITypeInfo::SetNested(const FWebAPITypeNameVariant& InNester)
 	ContainingType = InNester;
 }
 
+UObject* UWebAPITypeInfo::GetModel() const
+{
+	if (Model.IsNull())
+	{
+		return nullptr;
+	}
+
+	return Model.LoadSynchronous();
+}
+
 bool operator==(const UWebAPITypeInfo& A, const UWebAPITypeInfo& B)
 {
 	return (A.Name + A.Suffix) == (B.Name + B.Suffix)
diff --git a/Engine/Plugins/Experimental/Web/WebAPI/Source/WebAPIEditor/Private/Factories/WebAPIDefinitionFactory.cpp b/Engine/Plugins/Experimental/Web/WebAPI/Source/WebAPIEditor/Private/Factories/WebAPIDefinitionFactory.cpp
--- a/Engine/Plugins/Experimental/Web/WebAPI/Source/WebAPIEditor/Private/Factories/WebAPIDefinitionFactory.cpp
+++ b/Engine/Plugins/Experimental/Web/WebAPI/Source/WebAPIEditor/Private/Factories/WebAPIDefinitionFactory.cpp
@@ -108,7 +108,7 @@ TFuture<void> UWebAPIDefinitionFactory::PostImportWebAPI(UWebAPIDefinition* InDe
 		{
 			// check if either built-in type, or Model property is set
 			return InGeneratedType->bIsBuiltinType
-			|| !InGeneratedType->Model.IsNull()
+			|| InGeneratedType->GetModel()
 			|| InGeneratedType->SchemaType != EWebAPISchemaType::Model;
 		});
 
diff --git a/Engine/Plugins/Experimental/Web/WebAPI/Source/WebAPIEditor/Private/Tests/WebAPIEditorUtilities.spec.cpp b/Engine/Plugins/Experimental/Web/WebAPI/Source/WebAPIEditor/Private/Tests/WebAPIEditorUtilities.spec.cpp
--- a/Engine/Plugins/Experimental/Web/WebAPI/Source/WebAPIEditor/Private/Tests/WebAPIEditorUtilities.spec.cpp
+++ b/Engine/Plugins/Experimental/Web/WebAPI/Source/WebAPIEditor/Private/Tests/WebAPIEditorUtilities.spec.cpp
@@ -23,8 +23,9 @@ void FWebAPIEditorUtilitiesSpec::Define()
 			const FStringView TestString = TEXT("$teststring");
 			const FString ExpectedString = TEXT("teststring");
 
-			const FString ResultString = UE::WebAPI::FWebAPIStringUtilities::Get()->ToPascalCase(TestString);
-			
+			const FString ResultString = UE::WebAPI::FWebAPIStringUtilities::Get()->MakeValidMemberName(TestString);
+
+			// String equality is case-sensitive by default
 			TestEqual("String doesn't contain illegal characters", ResultString, ExpectedString);
 		});
 	});
diff --git a/Engine/Plugins/Experimental/Web/WebAPI/Source/WebAPIEditor/Public/Dom/WebAPISchema.h b/Engine/Plugins/Experimental/Web/WebAPI/Source/WebAPIEditor/Public/Dom/WebAPISchema.h
--- a/Engine/Plugins/Experimental/Web/WebAPI/Source/WebAPIEditor/Public/Dom/WebAPISchema.h
+++ b/Engine/Plugins/Experimental/Web/WebAPI/Source/WebAPIEditor/Public/Dom/WebAPISchema.h
@@ -326,13 +326,13 @@ public:
 
 	/** Add and return a reference to a new Model */
 	template <typename ModelType>
-	TObjectPtr<ModelType> AddModel(const TObjectPtr<UWebAPITypeInfo>& InTypeInfo = {});
+	ModelType* AddModel(const TObjectPtr<UWebAPITypeInfo>& InTypeInfo = {});
 
 	/** Add and return a reference to a new Enum. */
-	TObjectPtr<class UWebAPIEnum> AddEnum(const TObjectPtr<UWebAPITypeInfo>& InTypeInfo = {});
+	class UWebAPIEnum* AddEnum(const TObjectPtr<UWebAPITypeInfo>& InTypeInfo = {});
 
 	/** Add and return a reference to a new Parameter. These can be referenced by operations rather than defining them inline. */
-	TObjectPtr<class UWebAPIParameter> AddParameter(const TObjectPtr<UWebAPITypeInfo>& InTypeInfo);
+	class UWebAPIParameter* AddParameter(const TObjectPtr<UWebAPITypeInfo>& InTypeInfo);
 
 	/** Find/Add and return a reference to an existing/new Service. */
 	TObjectPtr<class UWebAPIService> GetOrMakeService(const FString& InName);
@@ -359,7 +359,7 @@ public:
 };
 
 template <typename ModelType>
-TObjectPtr<ModelType> UWebAPISchema::AddModel(const TObjectPtr<UWebAPITypeInfo>& InTypeInfo)
+ModelType* UWebAPISchema::AddModel(const TObjectPtr<UWebAPITypeInfo>& InTypeInfo)
 {
 	static_assert(TIsDerivedFrom<ModelType, UWebAPIModelBase>::Value, "Type is not derived from UWebAPIModelBase.");
 	
diff --git a/Engine/Plugins/Experimental/Web/WebAPI/Source/WebAPIEditor/Public/Dom/WebAPIType.h b/Engine/Plugins/Experimental/Web/WebAPI/Source/WebAPIEditor/Public/Dom/WebAPIType.h
--- a/Engine/Plugins/Experimental/Web/WebAPI/Source/WebAPIEditor/Public/Dom/WebAPIType.h
+++ b/Engine/Plugins/Experimental/Web/WebAPI/Source/WebAPIEditor/Public/Dom/WebAPIType.h
@@ -109,10 +109,6 @@ public:
 	UPROPERTY(VisibleAnywhere, Category = "Type")
 	FString DefaultValue;
 
-	/** Associated model, if any. */
-	UPROPERTY(VisibleAnywhere, Category = "Type")
-	TSoftObjectPtr<UWebAPIModelBase> Model;
-
 	/** Module dependencies for this type. */
 	UPROPERTY(VisibleAnywhere, Category = "Type")
 	TSet<FString> Modules = { TEXT("Core") };
@@ -157,13 +153,27 @@ public:
 	/** Set's this type as nested, using the provided Type as ContainingType. */
 	void SetNested(const FWebAPITypeNameVariant& InNester);
 
+	/** Get the (mutable) Model. Will return nullptr if not set. */
+	UObject* GetModel() const;
+
+	/** Set's the Model. */
+	template <
+		typename ModelType
+		UE_REQUIRES(
+			std::is_base_of_v<UObject, std::decay_t<ModelType>>
+			&& std::is_base_of_v<IWebAPISchemaObjectInterface, std::decay_t<ModelType>>)>
+	void SetModel(const ModelType* InModel)
+	{
+		Model = InModel;
+	}
+
 #if WITH_EDITOR
 	/** 
 	 * @return		Returns Valid if this object has data validation rules set up for it and the data for this object is valid. Returns Invalid if it does not pass the rules. Returns NotValidated if no rules are set for this object.
 	 */
 	virtual EDataValidationResult IsDataValid(class FDataValidationContext& Context) const override;
 #endif // WITH_EDITOR
-	
+
 	/** Comparisons. */
 	friend bool operator==(const UWebAPITypeInfo& A, const UWebAPITypeInfo& B);
 
@@ -171,4 +181,9 @@ public:
 	{
 		return !(A == B);
 	}
+
+private:
+	/** Associated model, if any. */
+	UPROPERTY(VisibleAnywhere, Category = "Type", meta = (MustImplement = "/Script/WebAPIEditor.WebAPISchemaObjectInterface"))
+	TSoftObjectPtr<UObject> Model;
 };
diff --git a/Engine/Plugins/Experimental/Web/WebAPI/Source/WebAPIOpenAPI/Private/Tests/WebAPIOpenAPI3.spec.cpp b/Engine/Plugins/Experimental/Web/WebAPI/Source/WebAPIOpenAPI/Private/Tests/WebAPIOpenAPI3.spec.cpp
--- a/Engine/Plugins/Experimental/Web/WebAPI/Source/WebAPIOpenAPI/Private/Tests/WebAPIOpenAPI3.spec.cpp
+++ b/Engine/Plugins/Experimental/Web/WebAPI/Source/WebAPIOpenAPI/Private/Tests/WebAPIOpenAPI3.spec.cpp
@@ -1,19 +1,19 @@
 // Copyright Epic Games, Inc. All Rights Reserved.
 
-#include "Misc/AutomationTest.h"
-#include "WebAPIDefinition.h"
-
 #if WITH_EDITOR
 
 #include "Dom/WebAPIEnum.h"
-#include "Dom/WebAPIParameter.h"
 #include "Dom/WebAPIOperation.h"
-#include "Misc/AutomationTest.h"
+#include "Dom/WebAPIParameter.h"
 #include "Dom/WebAPIService.h"
+#include "Interfaces/IPluginManager.h"
+#include "Misc/AutomationTest.h"
+#include "Misc/AutomationTest.h"
 #include "Misc/FileHelper.h"
 #include "Misc/Paths.h"
 #include "Serialization/JsonSerializer.h"
 #include "V3/WebAPIOpenAPIConverter.h"
+#include "WebAPIDefinition.h"
 
 #if WITH_DEV_AUTOMATION_TESTS
 
@@ -129,9 +129,11 @@ BEGIN_DEFINE_SPEC(FWebAPIOpenAPI3Spec,
 
 	FString GetSampleFile(const FString& InName) const
 	{
-		FString FilePath = FPaths::Combine(FPaths::EnginePluginsDir(),
-			TEXT("Web"),
-			TEXT("WebAPI"), TEXT("Source"), TEXT("WebAPIOpenAPI"),
+		TSharedPtr<IPlugin> Plugin = IPluginManager::Get().FindPlugin(UE_PLUGIN_NAME);
+		check(Plugin.IsValid());
+
+		FString FilePath = FPaths::Combine(Plugin->GetBaseDir(),
+			TEXT("Source"), TEXT("WebAPIOpenAPI"),
 			TEXT("Private"), TEXT("Tests"), TEXT("Samples"), TEXT("V3"), InName + TEXT(".json"));
 		ensure(FPaths::FileExists(FilePath));
 		return FilePath;
@@ -146,10 +148,20 @@ BEGIN_DEFINE_SPEC(FWebAPIOpenAPI3Spec,
 	{
 		FString FileContents;
 		FFileHelper::LoadFileToString(FileContents, *InFile);
-	
+
 		TSharedPtr<FJsonObject> JsonObject;
-		FJsonSerializer::Deserialize(TJsonReaderFactory<TCHAR>::Create(FileContents), JsonObject);
-		return JsonObject;
+		if (FJsonSerializer::Deserialize(TJsonReaderFactory<TCHAR>::Create(FileContents), JsonObject))
+		{
+			return JsonObject;
+		}
+	
+		TArray<TSharedPtr<FJsonValue>> JsonValues;
+		if (FJsonSerializer::Deserialize(TJsonReaderFactory<TCHAR>::Create(FileContents), JsonValues))
+		{
+			return JsonValues.IsEmpty() ? nullptr : JsonValues[0]->AsObject();
+		}
+
+		return nullptr;
 	}
 
 	TSharedPtr<FJsonValueArray> LoadJsonArray(const FString& InFile) const
diff --git a/Engine/Plugins/Experimental/Web/WebAPI/Source/WebAPIOpenAPI/Private/Tests/WebAPISwagger.spec.cpp b/Engine/Plugins/Experimental/Web/WebAPI/Source/WebAPIOpenAPI/Private/Tests/WebAPISwagger.spec.cpp
--- a/Engine/Plugins/Experimental/Web/WebAPI/Source/WebAPIOpenAPI/Private/Tests/WebAPISwagger.spec.cpp
+++ b/Engine/Plugins/Experimental/Web/WebAPI/Source/WebAPIOpenAPI/Private/Tests/WebAPISwagger.spec.cpp
@@ -2,6 +2,7 @@
 
 #if WITH_EDITOR
 
+#include "Interfaces/IPluginManager.h"
 #include "Misc/AutomationTest.h"
 #include "Misc/FileHelper.h"
 #include "Misc/Paths.h"
@@ -9,7 +10,6 @@
 #include "V2/WebAPISwaggerConverter.h"
 #include "V2/WebAPISwaggerSchema.h"
 #include "WebAPIDefinition.h"
-#include "WebAPIMessageLog.h"
 
 #if WITH_DEV_AUTOMATION_TESTS
 
@@ -23,9 +23,11 @@ BEGIN_DEFINE_SPEC(FWebAPISwaggerSpec,
 
 	FString GetSampleFile(const FString& InName) const
 	{
-		FString FilePath = FPaths::Combine(FPaths::EnginePluginsDir(),
-			TEXT("Web"),
-			TEXT("WebAPI"), TEXT("Source"), TEXT("WebAPIOpenAPI"),
+		TSharedPtr<IPlugin> Plugin = IPluginManager::Get().FindPlugin(UE_PLUGIN_NAME);
+		check(Plugin.IsValid());
+
+		FString FilePath = FPaths::Combine(Plugin->GetBaseDir(),
+			TEXT("Source"), TEXT("WebAPIOpenAPI"),
 			TEXT("Private"), TEXT("Tests"), TEXT("Samples"), TEXT("V2"), InName + TEXT(".json"));
 		ensure(FPaths::FileExists(FilePath));
 		return FilePath;
diff --git a/Engine/Plugins/Experimental/Web/WebAPI/Source/WebAPIOpenAPI/Private/V2/WebAPISwaggerConverter.cpp b/Engine/Plugins/Experimental/Web/WebAPI/Source/WebAPIOpenAPI/Private/V2/WebAPISwaggerConverter.cpp
--- a/Engine/Plugins/Experimental/Web/WebAPI/Source/WebAPIOpenAPI/Private/V2/WebAPISwaggerConverter.cpp
+++ b/Engine/Plugins/Experimental/Web/WebAPI/Source/WebAPIOpenAPI/Private/V2/WebAPISwaggerConverter.cpp
@@ -416,7 +416,7 @@ namespace UE::WebAPI::Swagger
 
 	template <typename SchemaType, typename ModelType>
 	bool FWebAPISwaggerSchemaConverter::ConvertModelBase(const TSharedPtr<SchemaType>& InSchema,
-														const TObjectPtr<ModelType>& OutModel)
+														ModelType* OutModel)
 	{
 		static_assert(TIsDerivedFrom<SchemaType, OpenAPI::V2::FSchemaBase>::Value, "Type is not derived from OpenAPI::V2::FSchemaBase.");
 
@@ -469,7 +469,7 @@ namespace UE::WebAPI::Swagger
 		DstEnum->Name = EnumTypeName;
 		DstEnum->Description = InSrcSchema->Description.Get(TEXT(""));
 
-		const TObjectPtr<UWebAPIModelBase> ModelBase = Cast<UWebAPIModelBase>(DstEnum);
+		UWebAPIModelBase* ModelBase = Cast<UWebAPIModelBase>(DstEnum);
 		if (!ConvertModelBase(InSrcSchema, ModelBase))
 		{
 			return nullptr;
@@ -493,7 +493,7 @@ namespace UE::WebAPI::Swagger
 		const FWebAPINameVariant& InPropertyName,
 		const TSharedPtr<OpenAPI::V2::FParameter>& InSchema,
 		const FString& InDefinitionName,
-		const TObjectPtr<UWebAPIProperty>& OutProperty)
+		UWebAPIProperty* OutProperty)
 	{
 		const FString DefinitionName = !InDefinitionName.IsEmpty()
 										? InDefinitionName
@@ -501,7 +501,7 @@ namespace UE::WebAPI::Swagger
 
 		OutProperty->Type = ResolveType(InSchema, DefinitionName);
 
-		const TObjectPtr<UWebAPIModelBase> ModelBase = OutProperty;
+		UWebAPIModelBase* ModelBase = OutProperty;
 		if (!ConvertModelBase(InSchema, ModelBase))
 		{
 			return false;
@@ -528,7 +528,7 @@ namespace UE::WebAPI::Swagger
 			}
 			EnumTypeName.TypeInfo->SetName(EnumName);
 
-			const TObjectPtr<UWebAPIEnum>& Enum = ConvertEnum(InSchema, EnumTypeName);
+			const UWebAPIEnum* Enum = ConvertEnum(InSchema, EnumTypeName);
 
 			const FText LogMessage = FText::FormatNamed(
 				LOCTEXT("AddedImplicitEnumForPropertyOfModel", "Implicit enum created for property \"{PropertyName}\" of model \"{ModelName}\"."),
@@ -540,7 +540,7 @@ namespace UE::WebAPI::Swagger
 			Enum->Name.TypeInfo->JsonName = InPropertyName.GetJsonName();
 
 			OutProperty->Type = Enum->Name;
-			OutProperty->Type.TypeInfo->Model = Enum;
+			OutProperty->Type.TypeInfo->SetModel(Enum);
 		}
 		// Add struct as it's own model, and reference it as this properties type
 		else if (OutProperty->Type.ToString(true).IsEmpty())
@@ -551,7 +551,7 @@ namespace UE::WebAPI::Swagger
 		return true;
 	}
 
-	bool FWebAPISwaggerSchemaConverter::PatchProperty(const FWebAPITypeNameVariant& InModelName, const FWebAPINameVariant& InPropertyName, const FWebAPITypeNameVariant& InPropertyTypeName, const TObjectPtr<UWebAPIProperty>& OutProperty)
+	bool FWebAPISwaggerSchemaConverter::PatchProperty(const FWebAPITypeNameVariant& InModelName, const FWebAPINameVariant& InPropertyName, const FWebAPITypeNameVariant& InPropertyTypeName, UWebAPIProperty* OutProperty)
 	{
 		OutProperty->Type = InPropertyTypeName;
 
@@ -567,7 +567,7 @@ namespace UE::WebAPI::Swagger
 		const FWebAPINameVariant& InParameterName,
 		const TSharedPtr<OpenAPI::V2::FParameter>& InParameter,
 		const FString& InDefinitionName,
-		const TObjectPtr<UWebAPIOperationParameter>& OutParameter)
+		UWebAPIOperationParameter* OutParameter)
 	{
 		// Will get schema or create if it doesn't exist (but will be empty)
 		Json::TJsonReference<OpenAPI::V2::FSchema> ParameterSchema = InParameter->Schema.Get({});
@@ -600,7 +600,7 @@ namespace UE::WebAPI::Swagger
 				ParameterSchema.GetShared());
 		}
 
-		const TObjectPtr<UWebAPIModelBase> ModelBase = Cast<UWebAPIModelBase>(OutParameter);
+		UWebAPIModelBase* ModelBase = Cast<UWebAPIModelBase>(OutParameter);
 		if (!ConvertModelBase(InParameter, ModelBase))
 		{
 			return false;
@@ -615,7 +615,7 @@ namespace UE::WebAPI::Swagger
 			ModelTypeName.TypeInfo->SetName(InDefinitionName);
 			ModelTypeName.TypeInfo->Prefix = TEXT("F");
 
-			const TObjectPtr<UWebAPIModel>& Model = OutputSchema->AddModel<UWebAPIModel>(
+			UWebAPIModel* Model = OutputSchema->AddModel<UWebAPIModel>(
 				ModelTypeName.HasTypeInfo() ? ModelTypeName.TypeInfo.Get() : nullptr);
 			PatchModel(InParameter->Schema.GetValue().GetShared(), {}, Model);
 
@@ -631,7 +631,7 @@ namespace UE::WebAPI::Swagger
 			Model->Name.TypeInfo->JsonType = UWebAPIStaticTypeRegistry::ToFromJsonType;
 
 			OutParameter->Type = Model->Name;
-			OutParameter->Type.TypeInfo->Model = Model;
+			OutParameter->Type.TypeInfo->SetModel(Model);
 		}
 
 		// Special case for "body" parameters
@@ -679,7 +679,7 @@ namespace UE::WebAPI::Swagger
 		// Will get schema or create if it doesn't exist (but will be empty)
 		const TSharedPtr<OpenAPI::V2::FSchema> SrcParameterSchema = ResolveReference(InSrcParameter->Schema.Get({}), SrcParameterDefinitionName, false);
 
-		const TObjectPtr<UWebAPIParameter> DstParameter = OutputSchema->AddParameter(ParameterTypeName.TypeInfo.Get());
+		UWebAPIParameter* DstParameter = OutputSchema->AddParameter(ParameterTypeName.TypeInfo.Get());
 
 		static TMap<FString, EWebAPIParameterStorage> InToStorage = {
 			{TEXT("query"), EWebAPIParameterStorage::Query},
@@ -700,7 +700,7 @@ namespace UE::WebAPI::Swagger
 				SrcParameterSchema);
 		}
 
-		const TObjectPtr<UWebAPIModel> Model = Cast<UWebAPIModel>(DstParameter);
+		UWebAPIModel* Model = Cast<UWebAPIModel>(DstParameter);
 		if (!PatchModel(InSrcParameter, ParameterTypeName, Model))
 		{
 			return nullptr;
@@ -725,7 +725,7 @@ namespace UE::WebAPI::Swagger
 				EnumTypeName.TypeInfo->SetName(ProviderSettings.MakeParameterTypeName(*ParameterName));
 				EnumTypeName.TypeInfo->SetNested(ParameterTypeName);
 
-				const TObjectPtr<UWebAPIEnum>& Enum = ConvertEnum(SrcParameterSchema, EnumTypeName);
+				const UWebAPIEnum* Enum = ConvertEnum(SrcParameterSchema, EnumTypeName);
 
 				const FText LogMessage = FText::FormatNamed(
 					LOCTEXT("AddedImplicitEnumForParameter", "Implicit enum created for parameter \"{Name}\"."),
@@ -735,7 +735,7 @@ namespace UE::WebAPI::Swagger
 				Enum->Name.TypeInfo->DebugString += LogMessage.ToString();
 
 				DstParameter->Type = Enum->Name;
-				DstParameter->Type.TypeInfo->Model = Enum;
+				DstParameter->Type.TypeInfo->SetModel(Enum);
 			}
 			// Add struct as it's own model, and reference it as this properties type
 			else if (DstParameter->Type.ToString(true).IsEmpty())
@@ -778,7 +778,7 @@ namespace UE::WebAPI::Swagger
 		return DstParameter;
 	}
 
-	bool FWebAPISwaggerSchemaConverter::ConvertRequest(const FWebAPITypeNameVariant& InOperationName, const TSharedPtr<OpenAPI::V2::FOperation>& InOperation, const TObjectPtr<UWebAPIOperationRequest>& OutRequest)
+	bool FWebAPISwaggerSchemaConverter::ConvertRequest(const FWebAPITypeNameVariant& InOperationName, const TSharedPtr<OpenAPI::V2::FOperation>& InOperation, UWebAPIOperationRequest* OutRequest)
 	{
 		const FWebAPITypeNameVariant RequestTypeName = OutputSchema->TypeRegistry->GetOrMakeGeneratedType(
 			EWebAPISchemaType::Model,
@@ -818,8 +818,8 @@ namespace UE::WebAPI::Swagger
 					SrcParameterDefinitionName = SrcParameterSchema->Items->GetLastPathSegment();
 				}
 
-				const TObjectPtr<UWebAPIOperationParameter> DstParameter = OutRequest->Parameters.Add_GetRef(NewObject<UWebAPIOperationParameter>(OutRequest));
-				ConvertOperationParameter(ParameterName, SrcParameter.GetShared(), SrcParameterDefinitionName,	DstParameter);
+				UWebAPIOperationParameter* DstParameter = OutRequest->Parameters.Add_GetRef(NewObject<UWebAPIOperationParameter>(OutRequest));
+				ConvertOperationParameter(ParameterName, SrcParameter.GetShared(), SrcParameterDefinitionName, DstParameter);
 				DstParameter->BindToTypeInfo();
 			
 				// Special case - if the name is "body", set as Body property and not Parameter
@@ -828,7 +828,7 @@ namespace UE::WebAPI::Swagger
 					// Check for existing definition
 					if (const TObjectPtr<UWebAPITypeInfo>* FoundGeneratedType = OutputSchema->TypeRegistry->FindGeneratedType(EWebAPISchemaType::Model, SrcParameterDefinitionName))
 					{
-						DstParameter->Model = Cast<UWebAPIModel>((*FoundGeneratedType)->Model.LoadSynchronous());
+						DstParameter->Model = Cast<UWebAPIModel>((*FoundGeneratedType)->GetModel());
 						return true;
 					}
 				}
@@ -852,7 +852,7 @@ namespace UE::WebAPI::Swagger
 		return true;
 	}
 
-	bool FWebAPISwaggerSchemaConverter::ConvertResponse(const FWebAPITypeNameVariant& InOperationName, uint32 InResponseCode, const TSharedPtr<OpenAPI::V2::FResponse>& InResponse, const TObjectPtr<UWebAPIOperationResponse>& OutResponse)
+	bool FWebAPISwaggerSchemaConverter::ConvertResponse(const FWebAPITypeNameVariant& InOperationName, uint32 InResponseCode, const TSharedPtr<OpenAPI::V2::FResponse>& InResponse, UWebAPIOperationResponse* OutResponse)
 	{
 		const FWebAPITypeNameVariant ResponseTypeName = OutputSchema->TypeRegistry->GetOrMakeGeneratedType(
 			EWebAPISchemaType::Model,
@@ -880,7 +880,7 @@ namespace UE::WebAPI::Swagger
 			SrcResponseSchema.Set(ResolveReference(SrcResponseSchema, SrcPropertyDefinitionName));
 
 			const FString PropertyName = SrcResponseSchema->Type.Get(GetDefaultJsonTypeForStorage(OutResponse->Storage)) == TEXT("array") ? ProviderSettings.GetDefaultArrayPropertyName() : ProviderSettings.GetDefaultPropertyName();
-			const TObjectPtr<UWebAPIProperty>& DstProperty = OutResponse->Properties.Add_GetRef(NewObject<UWebAPIProperty>(OutResponse));
+			UWebAPIProperty* DstProperty = OutResponse->Properties.Add_GetRef(NewObject<UWebAPIProperty>(OutResponse));
 			DstProperty->bIsMixin = true;
 			PatchProperty(OutResponse->Name,
 				PropertyName,
@@ -1022,24 +1022,24 @@ namespace UE::WebAPI::Swagger
 				Name,
 				TEXT("F"));
 
-			const TObjectPtr<UWebAPIModel>& Model = OutSchema->AddModel<UWebAPIModel>(ModelName.TypeInfo.Get());
+			UWebAPIModel* Model = OutSchema->AddModel<UWebAPIModel>(ModelName.TypeInfo.Get());
 			PatchModel(Schema, {}, Model);
 		}
 
 		// Set Model property of TypeInfos where applicable
 		for (const TObjectPtr<UWebAPIModelBase>& ModelBase : OutputSchema->Models)
 		{
-			if (const TObjectPtr<UWebAPIModel>& Model = Cast<UWebAPIModel>(ModelBase))
+			if (const UWebAPIModel* Model = Cast<UWebAPIModel>(ModelBase))
 			{
-				Model->Name.TypeInfo->Model = Model;
+				Model->Name.TypeInfo->SetModel(Model);
 			}
-			else if (const TObjectPtr<UWebAPIEnum>& Enum = Cast<UWebAPIEnum>(ModelBase))
+			else if (const UWebAPIEnum* Enum = Cast<UWebAPIEnum>(ModelBase))
 			{
-				Enum->Name.TypeInfo->Model = Enum;
+				Enum->Name.TypeInfo->SetModel(Enum);
 			}
-			else if (const TObjectPtr<UWebAPIParameter>& ParameterModel = Cast<UWebAPIParameter>(ModelBase))
+			else if (const UWebAPIParameter* ParameterModel = Cast<UWebAPIParameter>(ModelBase))
 			{
-				ParameterModel->Type.TypeInfo->Model = ParameterModel;
+				ParameterModel->Type.TypeInfo->SetModel(ParameterModel);
 			}
 		}
 
diff --git a/Engine/Plugins/Experimental/Web/WebAPI/Source/WebAPIOpenAPI/Private/V2/WebAPISwaggerConverter.h b/Engine/Plugins/Experimental/Web/WebAPI/Source/WebAPIOpenAPI/Private/V2/WebAPISwaggerConverter.h
--- a/Engine/Plugins/Experimental/Web/WebAPI/Source/WebAPIOpenAPI/Private/V2/WebAPISwaggerConverter.h
+++ b/Engine/Plugins/Experimental/Web/WebAPI/Source/WebAPIOpenAPI/Private/V2/WebAPISwaggerConverter.h
@@ -66,44 +66,41 @@ namespace UE::WebAPI::Swagger
 		bool IsArray(const TSharedPtr<SchemaType>& InSchema);
 
 		template <typename SchemaType, typename ModelType>
-		static bool ConvertModelBase(const TSharedPtr<SchemaType>& InSchema, const TObjectPtr<ModelType>& OutModel);
+		static bool ConvertModelBase(const TSharedPtr<SchemaType>& InSchema, ModelType* OutModel);
 
 		template <typename SchemaType>
 		TObjectPtr<UWebAPIEnum> ConvertEnum(const TSharedPtr<SchemaType>& InSrcSchema, const FWebAPITypeNameVariant& InEnumTypeName = {}) const;
 
 		/** Modifies an existing Property using the supplied source object. */
 		template <typename SchemaType>
-		bool PatchProperty(const TSharedPtr<SchemaType>& InSrcSchema, const FWebAPINameVariant& InPropertyName, const TObjectPtr<UWebAPIProperty>& OutProperty);
+		bool PatchProperty(const TSharedPtr<SchemaType>& InSrcSchema, const FWebAPINameVariant& InPropertyName, UWebAPIProperty* OutProperty);
 
 		/** Modifies an existing Property using the supplied source object. */
 		template <typename SchemaType>
-		bool PatchProperty(const FWebAPITypeNameVariant& InModelName, const FWebAPINameVariant& InPropertyName, const TSharedPtr<SchemaType>& InSchema, const FString& InDefinitionName, const TObjectPtr<UWebAPIProperty>& OutProperty);
+		bool PatchProperty(const FWebAPITypeNameVariant& InModelName, const FWebAPINameVariant& InPropertyName, const TSharedPtr<SchemaType>& InSchema, const FString& InDefinitionName, UWebAPIProperty* OutProperty);
 
 		/** Modifies an existing Property using the supplied TypeInfo. */
-		bool PatchProperty(const FWebAPITypeNameVariant& InModelName, const FWebAPINameVariant& InPropertyName, const FWebAPITypeNameVariant& InPropertyTypeName, const TObjectPtr<UWebAPIProperty>& OutProperty);
+		bool PatchProperty(const FWebAPITypeNameVariant& InModelName, const FWebAPINameVariant& InPropertyName, const FWebAPITypeNameVariant& InPropertyTypeName, UWebAPIProperty* OutProperty);
 		
 		/** Creates a new Property from the supplied source object. */
 		template <typename SchemaType>
 		TObjectPtr<UWebAPIProperty> ConvertProperty(const TSharedPtr<SchemaType>& InSrcSchema, const TObjectPtr<UWebAPIModel>& InModel, const FWebAPINameVariant& InPropertyName = {}, const FString& InDefinitionName = {});
-/*
-		template <>
-		bool ConvertProperty(const FWebAPITypeNameVariant& InModelName,	const FWebAPINameVariant& InPropertyName, const TSharedPtr<OpenAPI::V2::FParameter>& InSchema, const FString& InDefinitionName, const TObjectPtr<UWebAPIProperty>& OutProperty);
-*/
+
 		/** Modifies an existing Model using the supplied source object. */
 		template <typename SchemaType>
-		bool PatchModel(const TSharedPtr<SchemaType>& InSrcSchema, const FWebAPITypeNameVariant& InModelTypeName, const TObjectPtr<UWebAPIModel>& OutModel);
+		bool PatchModel(const TSharedPtr<SchemaType>& InSrcSchema, const FWebAPITypeNameVariant& InModelTypeName, UWebAPIModel* OutModel);
 		
 		/** Creates a new Model from the supplied source object. */
 		template <typename SchemaType>
 		TObjectPtr<UWebAPIModel> ConvertModel(const TSharedPtr<SchemaType>& InSrcSchema, const FWebAPITypeNameVariant& InModelTypeName = {});
 
-		bool ConvertOperationParameter(const FWebAPINameVariant& InParameterName, const TSharedPtr<OpenAPI::V2::FParameter>& InParameter, const FString& InDefinitionName, const TObjectPtr<UWebAPIOperationParameter>& OutParameter);
+		bool ConvertOperationParameter(const FWebAPINameVariant& InParameterName, const TSharedPtr<OpenAPI::V2::FParameter>& InParameter, const FString& InDefinitionName, UWebAPIOperationParameter* OutParameter);
 
 		TObjectPtr<UWebAPIParameter> ConvertParameter(const TSharedPtr<OpenAPI::V2::FParameter>& InSrcParameter);
 		
-		bool ConvertRequest(const FWebAPITypeNameVariant& InOperationName, const TSharedPtr<OpenAPI::V2::FOperation>& InOperation, const TObjectPtr<UWebAPIOperationRequest>& OutRequest);
+		bool ConvertRequest(const FWebAPITypeNameVariant& InOperationName, const TSharedPtr<OpenAPI::V2::FOperation>& InOperation,  UWebAPIOperationRequest* OutRequest);
 
-		bool ConvertResponse(const FWebAPITypeNameVariant& InOperationName, uint32 InResponseCode, const TSharedPtr<OpenAPI::V2::FResponse>& InResponse, const TObjectPtr<UWebAPIOperationResponse>& OutResponse);
+		bool ConvertResponse(const FWebAPITypeNameVariant& InOperationName, uint32 InResponseCode, const TSharedPtr<OpenAPI::V2::FResponse>& InResponse, UWebAPIOperationResponse* OutResponse);
 
 		TObjectPtr<UWebAPIOperation> ConvertOperation(const FString& InPath, const FString& InVerb, const TSharedPtr<OpenAPI::V2::FOperation>& InSrcOperation, const FWebAPITypeNameVariant& InOperationTypeName = {});
 
@@ -134,7 +131,7 @@ namespace UE::WebAPI::Swagger
 
 	/** Modifies an existing Model using the supplied source object. */
 	template <>
-	bool FWebAPISwaggerSchemaConverter::PatchModel(const TSharedPtr<OpenAPI::V2::FParameter>& InSrcSchema, const FWebAPITypeNameVariant& InModelTypeName, const TObjectPtr<UWebAPIModel>& OutModel);
+	bool FWebAPISwaggerSchemaConverter::PatchModel(const TSharedPtr<OpenAPI::V2::FParameter>& InSrcSchema, const FWebAPITypeNameVariant& InModelTypeName, UWebAPIModel* OutModel);
 }
 
 #include "WebAPISwaggerConverter.inl"
diff --git a/Engine/Plugins/Experimental/Web/WebAPI/Source/WebAPIOpenAPI/Private/V2/WebAPISwaggerConverter.inl b/Engine/Plugins/Experimental/Web/WebAPI/Source/WebAPIOpenAPI/Private/V2/WebAPISwaggerConverter.inl
--- a/Engine/Plugins/Experimental/Web/WebAPI/Source/WebAPIOpenAPI/Private/V2/WebAPISwaggerConverter.inl
+++ b/Engine/Plugins/Experimental/Web/WebAPI/Source/WebAPIOpenAPI/Private/V2/WebAPISwaggerConverter.inl
@@ -5,7 +5,6 @@
 #include "Dom/WebAPIEnum.h"
 #include "Dom/WebAPIType.h"
 #include "Dom/WebAPITypeRegistry.h"
-#include "IWebAPIEditorModule.h"
 #include "V2/WebAPISwaggerProvider.h"
 
 #define LOCTEXT_NAMESPACE "WebAPISwaggerConverter"
@@ -17,7 +16,7 @@ namespace UE::WebAPI::Swagger
 														const FWebAPINameVariant& InPropertyName,
 														const TSharedPtr<SchemaType>& InSchema,
 														const FString& InDefinitionName,
-														const TObjectPtr<UWebAPIProperty>& OutProperty)
+														UWebAPIProperty* OutProperty)
 	{
 		static_assert(TIsDerivedFrom<SchemaType, OpenAPI::V2::FSchemaBase>::Value, "Type is not derived from OpenAPI::V2::FSchemaBase.");
 
@@ -37,8 +36,8 @@ namespace UE::WebAPI::Swagger
 
 		OutProperty->Type = ResolveType<SchemaType>(ItemSchema, DefinitionName);
 
-		const TObjectPtr<UWebAPIModelBase> ModelBase = OutProperty;
-		if (!ConvertModelBase(ItemSchema, ModelBase))
+		UWebAPIModelBase* ModelBase = OutProperty;
+		if (!ConvertModelBase<SchemaType, UWebAPIModelBase>(ItemSchema, ModelBase))
 		{
 			return false;
 		}
@@ -56,7 +55,7 @@ namespace UE::WebAPI::Swagger
 			const FWebAPITypeNameVariant EnumTypeName = OutProperty->Type;
 			EnumTypeName.TypeInfo->SetName(ProviderSettings.MakeNestedPropertyTypeName(InModelName, InPropertyName));
 
-			const TObjectPtr<UWebAPIEnum>& Enum = ConvertEnum(InSchema, EnumTypeName);
+			const UWebAPIEnum* Enum = ConvertEnum(InSchema, EnumTypeName);
 
 			const FText LogMessage = FText::FormatNamed(
 				LOCTEXT("AddedImplicitEnumForPropertyOfModel", "Implicit enum created for property \"{PropertyName}\" of model \"{ModelName}\"."),
@@ -70,7 +69,7 @@ namespace UE::WebAPI::Swagger
 			Enum->Name.TypeInfo->SetNested(InModelName);
 
 			OutProperty->Type = Enum->Name;
-			OutProperty->Type.TypeInfo->Model = Enum;
+			OutProperty->Type.TypeInfo->SetModel(Enum);
 		}
 		// Add struct as it's own model, and reference it as this properties type
 		else if (OutProperty->Type.ToString(true).IsEmpty())
@@ -80,7 +79,7 @@ namespace UE::WebAPI::Swagger
 			ModelTypeName.TypeInfo->Prefix = TEXT("F");
 			ModelTypeName.TypeInfo->SetNested(InModelName);
 
-			const TObjectPtr<UWebAPIModel>& Model = OutputSchema->AddModel<UWebAPIModel>(ModelTypeName.HasTypeInfo() ? ModelTypeName.TypeInfo.Get() : nullptr);
+			UWebAPIModel* Model = OutputSchema->AddModel<UWebAPIModel>(ModelTypeName.HasTypeInfo() ? ModelTypeName.TypeInfo.Get() : nullptr);
 			PatchModel(InSchema, {}, Model);
 
 			const FText LogMessage = FText::FormatNamed(
@@ -95,7 +94,7 @@ namespace UE::WebAPI::Swagger
 			Model->Name.TypeInfo->JsonType = UWebAPIStaticTypeRegistry::ToFromJsonType;
 
 			OutProperty->Type = Model->Name;
-			OutProperty->Type.TypeInfo->Model = Model;
+			OutProperty->Type.TypeInfo->SetModel(Model);
 		}
 		
 		OutProperty->Name = ResolvePropertyName(OutProperty, InModelName);
@@ -104,184 +103,183 @@ namespace UE::WebAPI::Swagger
 	}
 
 		template <typename SchemaType>
-			TObjectPtr<UWebAPIProperty> FWebAPISwaggerSchemaConverter::ConvertProperty(const TSharedPtr<SchemaType>& InSrcSchema, const TObjectPtr<UWebAPIModel>& InModel, const FWebAPINameVariant& InPropertyName, const FString& InDefinitionName)
+		TObjectPtr<UWebAPIProperty> FWebAPISwaggerSchemaConverter::ConvertProperty(const TSharedPtr<SchemaType>& InSrcSchema, const TObjectPtr<UWebAPIModel>& InModel, const FWebAPINameVariant& InPropertyName, const FString& InDefinitionName)
+		{
+			static_assert(TIsDerivedFrom<SchemaType, OpenAPI::V2::FSchemaBase>::Value, "Type is not derived from OpenAPI::V2::FSchemaBase.");
+
+			const FWebAPINameVariant SrcPropertyName = InPropertyName;
+			FString SrcPropertyDefinitionName = ProviderSettings.MakeNestedPropertyTypeName(*InModel->Name.ToString(true), SrcPropertyName);
+			const TSharedPtr<OpenAPI::V2::FSchema> SrcPropertyValue = ResolveReference(InSrcSchema, SrcPropertyDefinitionName);
+
+			if (SrcPropertyValue)
 			{
-				static_assert(TIsDerivedFrom<SchemaType, OpenAPI::V2::FSchemaBase>::Value, "Type is not derived from OpenAPI::V2::FSchemaBase.");
+				const TObjectPtr<UWebAPIProperty>& DstProperty = InModel->Properties.Add_GetRef(NewObject<UWebAPIProperty>(InModel));
+				PatchProperty(InModel->Name,
+					FWebAPINameInfo(NameTransformer(SrcPropertyName), SrcPropertyName.GetJsonName(), InModel->Name),
+					SrcPropertyValue,
+					SrcPropertyDefinitionName,
+					DstProperty);
+
+				return DstProperty;
+			}
 
-				const FWebAPINameVariant SrcPropertyName = InPropertyName;
-				FString SrcPropertyDefinitionName = ProviderSettings.MakeNestedPropertyTypeName(*InModel->Name.ToString(true), SrcPropertyName);
-				const TSharedPtr<OpenAPI::V2::FSchema> SrcPropertyValue = ResolveReference(InSrcSchema, SrcPropertyDefinitionName);
+			return nullptr;
+		}
 
-				if (SrcPropertyValue)
-				{
-					const TObjectPtr<UWebAPIProperty>& DstProperty = InModel->Properties.Add_GetRef(NewObject<UWebAPIProperty>(InModel));
-					PatchProperty(InModel->Name,
-						FWebAPINameInfo(NameTransformer(SrcPropertyName), SrcPropertyName.GetJsonName(), InModel->Name),
-						SrcPropertyValue,
-						SrcPropertyDefinitionName,
-						DstProperty);
-					
-					return DstProperty;
-				}
+		template <typename SchemaType>
+		bool FWebAPISwaggerSchemaConverter::PatchModel(const TSharedPtr<SchemaType>& InSrcSchema, const FWebAPITypeNameVariant& InModelTypeName, UWebAPIModel* OutModel)
+		{
+			static_assert(TIsDerivedFrom<SchemaType, OpenAPI::V2::FSchemaBase>::Value, "Type is not derived from OpenAPI::V2::FSchemaBase.");
 
-				return nullptr;
+			FWebAPITypeNameVariant ModelTypeName;
+			if(InModelTypeName.IsValid())
+			{
+				ModelTypeName = InModelTypeName;
+			}
+			else if(OutModel->Name.HasTypeInfo())
+			{
+				ModelTypeName = OutModel->Name;
+			}
+			else
+			{
+				const FString ModelName = InSrcSchema->Name;
+				check(!ModelName.IsEmpty());
+
+				ModelTypeName = OutputSchema->TypeRegistry->GetOrMakeGeneratedType(
+					EWebAPISchemaType::Model,
+					NameTransformer(*ModelName),
+					*ModelName,
+					TEXT("F"));
 			}
 
-			template <typename SchemaType>
-			bool FWebAPISwaggerSchemaConverter::PatchModel(const TSharedPtr<SchemaType>& InSrcSchema, const FWebAPITypeNameVariant& InModelTypeName, const TObjectPtr<UWebAPIModel>& OutModel)
+			if (!InSrcSchema.IsValid())
 			{
-				static_assert(TIsDerivedFrom<SchemaType, OpenAPI::V2::FSchemaBase>::Value, "Type is not derived from OpenAPI::V2::FSchemaBase.");
+				FFormatNamedArguments Args;
+				Args.Add(TEXT("ModelName"), FText::FromString(*ModelTypeName.ToString(true)));
+				MessageLog->LogWarning(
+					FText::Format(LOCTEXT("SchemaInvalid",
+							"The schema for model \"{ModelName}\" was invalid/null."),
+						Args),
+					FWebAPISwaggerProvider::LogName);
+				return false;
+			}
 
-				FWebAPITypeNameVariant ModelTypeName;
-				if(InModelTypeName.IsValid())
-				{
-					ModelTypeName = InModelTypeName;
-				}
-				else if(OutModel->Name.HasTypeInfo())
-				{
-					ModelTypeName = OutModel->Name;
-				}
-				else
-				{
-					const FString ModelName = InSrcSchema->Name;
-					check(!ModelName.IsEmpty());
-
-					ModelTypeName = OutputSchema->TypeRegistry->GetOrMakeGeneratedType(
-						EWebAPISchemaType::Model,
-						NameTransformer(*ModelName),
-						*ModelName,
-						TEXT("F"));
-				}
+			const TObjectPtr<UWebAPIModel>& Model = OutModel ? OutModel : OutputSchema->AddModel<UWebAPIModel>(ModelTypeName.TypeInfo.Get());
 
-				if (!InSrcSchema.IsValid())
-				{
-					FFormatNamedArguments Args;
-					Args.Add(TEXT("ModelName"), FText::FromString(*ModelTypeName.ToString(true)));
-					MessageLog->LogWarning(
-						FText::Format(LOCTEXT("SchemaInvalid",
-								"The schema for model \"{ModelName}\" was invalid/null."),
-							Args),
-						FWebAPISwaggerProvider::LogName);
-					return false;
-				}
+			UWebAPIModelBase* ModelBase = Model;
+			if (!ConvertModelBase(InSrcSchema, ModelBase))
+			{
+				return false;
+			}
 
-				const TObjectPtr<UWebAPIModel>& Model = OutModel ? OutModel : OutputSchema->AddModel<UWebAPIModel>(ModelTypeName.TypeInfo.Get());
+			Model->Name = ModelTypeName;
 
-				const TObjectPtr<UWebAPIModelBase> ModelBase = Model;
-				if (!ConvertModelBase(InSrcSchema, ModelBase))
+			if (InSrcSchema->Properties.IsSet() && !InSrcSchema->Properties->IsEmpty())
+			{
+				for (const TTuple<FString, Json::TJsonReference<OpenAPI::V2::FSchema>>& NamePropertyPair : InSrcSchema->Properties.GetValue())
 				{
-					return false;
-				}
-
-				Model->Name = ModelTypeName;
+					const FString& SrcPropertyName = NamePropertyPair.Key;
+					FString SrcPropertyDefinitionName = ProviderSettings.MakeNestedPropertyTypeName(*ModelTypeName.ToString(true), SrcPropertyName);
+					TSharedPtr<OpenAPI::V2::FSchema> SrcPropertyValue = ResolveReference(NamePropertyPair.Value, SrcPropertyDefinitionName);
 
-				if (InSrcSchema->Properties.IsSet() && !InSrcSchema->Properties->IsEmpty())
-				{
-					for (const TTuple<FString, Json::TJsonReference<OpenAPI::V2::FSchema>>& NamePropertyPair : InSrcSchema->Properties.GetValue())
+					if (SrcPropertyValue)
 					{
-						const FString& SrcPropertyName = NamePropertyPair.Key;
-						FString SrcPropertyDefinitionName = ProviderSettings.MakeNestedPropertyTypeName(*ModelTypeName.ToString(true), SrcPropertyName);
-						TSharedPtr<OpenAPI::V2::FSchema> SrcPropertyValue = ResolveReference(NamePropertyPair.Value, SrcPropertyDefinitionName);
-
-						if (SrcPropertyValue)
-						{
-							TObjectPtr<UWebAPIProperty>& DstProperty = Model->Properties.Add_GetRef(NewObject<UWebAPIProperty>(Model));
-							PatchProperty(Model->Name,
-								FWebAPINameInfo(NameTransformer(SrcPropertyName), SrcPropertyName, ModelTypeName),
-								SrcPropertyValue,
-								SrcPropertyDefinitionName,
-								DstProperty);
-						}
+						TObjectPtr<UWebAPIProperty>& DstProperty = Model->Properties.Add_GetRef(NewObject<UWebAPIProperty>(Model));
+						PatchProperty(Model->Name,
+							FWebAPINameInfo(NameTransformer(SrcPropertyName), SrcPropertyName, ModelTypeName),
+							SrcPropertyValue,
+							SrcPropertyDefinitionName,
+							DstProperty);
 					}
 				}
+			}
 
-				Model->BindToTypeInfo();
+			Model->BindToTypeInfo();
 
-				return true;
+			return true;
+		}
+
+		template <>
+		inline bool FWebAPISwaggerSchemaConverter::PatchModel<OpenAPI::V2::FParameter>(const TSharedPtr<OpenAPI::V2::FParameter>& InSrcSchema, const FWebAPITypeNameVariant& InModelTypeName, UWebAPIModel* OutModel)
+		{
+			FWebAPITypeNameVariant ModelTypeName;
+			if(InModelTypeName.IsValid() && InModelTypeName.HasTypeInfo())
+			{
+				ModelTypeName = InModelTypeName;
+			}
+			else if(OutModel->Name.HasTypeInfo())
+			{
+				ModelTypeName = OutModel->Name;
+			}
+			else
+			{
+				const FString ModelName = InSrcSchema->Name;
+				check(!ModelName.IsEmpty());
+
+				ModelTypeName = OutputSchema->TypeRegistry->GetOrMakeGeneratedType(
+					EWebAPISchemaType::Model,
+					NameTransformer(*ModelName),
+					*ModelName,
+					TEXT("F"));
 			}
 
-			template <>
-			inline bool FWebAPISwaggerSchemaConverter::PatchModel<OpenAPI::V2::FParameter>(const TSharedPtr<OpenAPI::V2::FParameter>& InSrcSchema, const FWebAPITypeNameVariant& InModelTypeName, const TObjectPtr<UWebAPIModel>& OutModel)
+			if (!InSrcSchema.IsValid())
 			{
-				FWebAPITypeNameVariant ModelTypeName;
-				if(InModelTypeName.IsValid() && InModelTypeName.HasTypeInfo())
-				{
-					ModelTypeName = InModelTypeName;
-				}
-				else if(OutModel->Name.HasTypeInfo())
-				{
-					ModelTypeName = OutModel->Name;
-				}
-				else
-				{
-					const FString ModelName = InSrcSchema->Name;
-					check(!ModelName.IsEmpty());
-
-					ModelTypeName = OutputSchema->TypeRegistry->GetOrMakeGeneratedType(
-						EWebAPISchemaType::Model,
-						NameTransformer(*ModelName),
-						*ModelName,
-						TEXT("F"));
-				}
+				FFormatNamedArguments Args;
+				Args.Add(TEXT("ModelName"), FText::FromString(*ModelTypeName.ToString(true)));
+				MessageLog->LogWarning(
+					FText::Format(LOCTEXT("SchemaInvalid",
+							"The schema for model \"{ModelName}\" was invalid/null."),
+						Args),
+					FWebAPISwaggerProvider::LogName);
+				return false;
+			}
 
-				if (!InSrcSchema.IsValid())
-				{
-					FFormatNamedArguments Args;
-					Args.Add(TEXT("ModelName"), FText::FromString(*ModelTypeName.ToString(true)));
-					MessageLog->LogWarning(
-						FText::Format(LOCTEXT("SchemaInvalid",
-								"The schema for model \"{ModelName}\" was invalid/null."),
-							Args),
-						FWebAPISwaggerProvider::LogName);
-					return false;
-				}
+			UWebAPIModel* Model = OutModel ? OutModel : OutputSchema->AddModel<UWebAPIModel>(ModelTypeName.TypeInfo.Get());
 
-				const TObjectPtr<UWebAPIModel>& Model = OutModel ? OutModel : OutputSchema->AddModel<UWebAPIModel>(ModelTypeName.TypeInfo.Get());
+			UWebAPIModelBase* ModelBase = Model;
+			if (!ConvertModelBase<OpenAPI::V2::FParameter, UWebAPIModelBase>(InSrcSchema, ModelBase))
+			{
+				return false;
+			}
 
-				const TObjectPtr<UWebAPIModelBase> ModelBase = Model;
-				if (!ConvertModelBase(InSrcSchema, ModelBase))
-				{
-					return false;
-				}
+			Model->Name = ModelTypeName;
 
-				Model->Name = ModelTypeName;
+			Model->BindToTypeInfo();
 
-				Model->BindToTypeInfo();
+			return true;
+		}
 
-				return true;
-			}
+		template <typename SchemaType>
+		TObjectPtr<UWebAPIModel> FWebAPISwaggerSchemaConverter::ConvertModel(const TSharedPtr<SchemaType>& InSrcSchema, const FWebAPITypeNameVariant& InModelTypeName)
+		{
+			check(InSrcSchema.IsValid());
 
-			template <typename SchemaType>
-			TObjectPtr<UWebAPIModel> FWebAPISwaggerSchemaConverter::ConvertModel(const TSharedPtr<SchemaType>& InSrcSchema, const FWebAPITypeNameVariant& InModelTypeName)
+			FWebAPITypeNameVariant ModelTypeName;
+			if(InModelTypeName.IsValid() && ModelTypeName.HasTypeInfo())
 			{
-				check(InSrcSchema.IsValid());
-
-				FWebAPITypeNameVariant ModelTypeName;
-				if(InModelTypeName.IsValid() && ModelTypeName.HasTypeInfo())
-				{
-					ModelTypeName = InModelTypeName;
-				}
-				else
-				{
-					const FString ModelName = InSrcSchema->Name;
-					check(!ModelName.IsEmpty());
-
-					ModelTypeName = OutputSchema->TypeRegistry->GetOrMakeGeneratedType(
-						EWebAPISchemaType::Model,
-						NameTransformer(*ModelName),
-						*ModelName,
-						TEXT("F"));
-				}
-				
-				const TObjectPtr<UWebAPIModel>& Model = OutputSchema->AddModel<UWebAPIModel>(ModelTypeName.TypeInfo.Get());
-				if(PatchModel(InSrcSchema, ModelTypeName, Model))
-				{
-					return Model;
-				}
+				ModelTypeName = InModelTypeName;
+			}
+			else
+			{
+				const FString ModelName = InSrcSchema->Name;
+				check(!ModelName.IsEmpty());
+
+				ModelTypeName = OutputSchema->TypeRegistry->GetOrMakeGeneratedType(
+					EWebAPISchemaType::Model,
+					NameTransformer(*ModelName),
+					*ModelName,
+					TEXT("F"));
+			}
 
-				return nullptr;
+			const TObjectPtr<UWebAPIModel>& Model = OutputSchema->AddModel<UWebAPIModel>(ModelTypeName.TypeInfo.Get());
+			if(PatchModel(InSrcSchema, ModelTypeName, Model))
+			{
+				return Model;
 			}
-	
+
+			return nullptr;
+		}
 }
 
 #undef LOCTEXT_NAMESPACE
diff --git a/Engine/Plugins/Experimental/Web/WebAPI/Source/WebAPIOpenAPI/Private/V3/WebAPIOpenAPIConverter.cpp b/Engine/Plugins/Experimental/Web/WebAPI/Source/WebAPIOpenAPI/Private/V3/WebAPIOpenAPIConverter.cpp
--- a/Engine/Plugins/Experimental/Web/WebAPI/Source/WebAPIOpenAPI/Private/V3/WebAPIOpenAPIConverter.cpp
+++ b/Engine/Plugins/Experimental/Web/WebAPI/Source/WebAPIOpenAPI/Private/V3/WebAPIOpenAPIConverter.cpp
@@ -292,7 +292,7 @@ namespace UE::WebAPI::OpenAPI
 		return InProperty->Name;
 	}
 
-	bool FWebAPIOpenAPISchemaConverter::ConvertProperty(const FWebAPITypeNameVariant& InModelName, const FWebAPINameVariant& InPropertyName, const FWebAPITypeNameVariant& InPropertyTypeName, const TObjectPtr<UWebAPIProperty>& OutProperty)
+	bool FWebAPIOpenAPISchemaConverter::ConvertProperty(const FWebAPITypeNameVariant& InModelName, const FWebAPINameVariant& InPropertyName, const FWebAPITypeNameVariant& InPropertyTypeName, UWebAPIProperty* OutProperty)
 	{
 		OutProperty->Type = InPropertyTypeName;
 		check(OutProperty->Type.IsValid());
@@ -309,7 +309,7 @@ namespace UE::WebAPI::OpenAPI
 		const FWebAPINameVariant& InParameterName,
 		const TSharedPtr<OpenAPI::V3::FParameterObject>& InParameter,
 		const FString& InDefinitionName,
-		const TObjectPtr<UWebAPIOperationParameter>& OutParameter)
+		UWebAPIOperationParameter* OutParameter)
 	{
 		// Will get schema or create if it doesn't exist (but will be empty)
 		const TSharedPtr<V3::FSchemaObject> ParameterSchema = ResolveReference(InParameter->Schema.Get({}));
@@ -331,7 +331,7 @@ namespace UE::WebAPI::OpenAPI
 		
 		OutParameter->bIsArray = IsArray(InParameter);
 
-		const TObjectPtr<UWebAPIModelBase> ModelBase = Cast<UWebAPIModelBase>(OutParameter);
+		UWebAPIModelBase* ModelBase = Cast<UWebAPIModelBase>(OutParameter);
 		if (!ConvertModelBase(InParameter, ModelBase))
 		{
 			return false;
@@ -346,7 +346,7 @@ namespace UE::WebAPI::OpenAPI
 			ModelTypeName.TypeInfo->SetName(InDefinitionName);
 			ModelTypeName.TypeInfo->Prefix = TEXT("F");
 
-			const TObjectPtr<UWebAPIModel>& Model = OutputSchema->AddModel<UWebAPIModel>(
+			UWebAPIModel* Model = OutputSchema->AddModel<UWebAPIModel>(
 				ModelTypeName.HasTypeInfo() ? ModelTypeName.TypeInfo.Get() : nullptr);
 			ConvertModel(InParameter->Schema.GetValue().GetShared(), {}, Model);
 
@@ -362,7 +362,7 @@ namespace UE::WebAPI::OpenAPI
 			Model->Name.TypeInfo->JsonType = UWebAPIStaticTypeRegistry::ToFromJsonType;
 
 			OutParameter->Type = Model->Name;
-			OutParameter->Type.TypeInfo->Model = Model;
+			OutParameter->Type.TypeInfo->SetModel(Model);
 		}
 
 		// Special case for "body" parameters
@@ -462,7 +462,7 @@ namespace UE::WebAPI::OpenAPI
 				EnumTypeName.TypeInfo->SetName(ProviderSettings.MakeParameterTypeName(*ParameterName));
 				EnumTypeName.TypeInfo->SetNested(ParameterTypeName);
 
-				const TObjectPtr<UWebAPIEnum>& Enum = ConvertEnum(SrcParameterSchema, EnumTypeName);
+				const UWebAPIEnum* Enum = ConvertEnum(SrcParameterSchema, EnumTypeName);
 
 				const FText LogMessage = FText::FormatNamed(
 					LOCTEXT("AddedImplicitEnumForParameter", "Implicit enum created for parameter \"{Name}\"."),
@@ -472,7 +472,7 @@ namespace UE::WebAPI::OpenAPI
 				Enum->Name.TypeInfo->DebugString += LogMessage.ToString();
 
 				DstParameter->Type = Enum->Name;
-				DstParameter->Type.TypeInfo->Model = Enum;
+				DstParameter->Type.TypeInfo->SetModel(Enum);
 			}
 			// Add struct as it's own model, and reference it as this properties type
 			else if (DstParameter->Type.ToString(true).IsEmpty())
@@ -517,7 +517,7 @@ namespace UE::WebAPI::OpenAPI
 		return DstParameter;
 	}
 
-	bool FWebAPIOpenAPISchemaConverter::ConvertRequest(const FWebAPITypeNameVariant& InOperationName, const TSharedPtr<OpenAPI::V3::FOperationObject>& InOperation, const TObjectPtr<UWebAPIOperationRequest>& OutRequest)
+	bool FWebAPIOpenAPISchemaConverter::ConvertRequest(const FWebAPITypeNameVariant& InOperationName, const TSharedPtr<OpenAPI::V3::FOperationObject>& InOperation, UWebAPIOperationRequest* OutRequest)
 	{
 		const FWebAPITypeNameVariant RequestTypeName = OutputSchema->TypeRegistry->GetOrMakeGeneratedType(
 			EWebAPISchemaType::Model,
@@ -555,7 +555,7 @@ namespace UE::WebAPI::OpenAPI
 					// Check for existing definition
 					if (const TObjectPtr<UWebAPITypeInfo>* FoundGeneratedType = OutputSchema->TypeRegistry->FindGeneratedType(EWebAPISchemaType::Model, SrcParameterDefinitionName))
 					{
-						DstParameter->Model = Cast<UWebAPIModel>((*FoundGeneratedType)->Model.LoadSynchronous());
+						DstParameter->Model = Cast<UWebAPIModel>((*FoundGeneratedType)->GetModel());
 						return true;
 					}
 				}
@@ -615,7 +615,7 @@ namespace UE::WebAPI::OpenAPI
 		return true;
 	}
 
-	bool FWebAPIOpenAPISchemaConverter::ConvertResponse(const FWebAPITypeNameVariant& InOperationName, uint32 InResponseCode, const TSharedPtr<OpenAPI::V3::FResponseObject>& InResponse, const TObjectPtr<UWebAPIOperationResponse>& OutResponse)
+	bool FWebAPIOpenAPISchemaConverter::ConvertResponse(const FWebAPITypeNameVariant& InOperationName, uint32 InResponseCode, const TSharedPtr<OpenAPI::V3::FResponseObject>& InResponse, UWebAPIOperationResponse* OutResponse)
 	{
 		const FWebAPITypeNameVariant ResponseTypeName = OutputSchema->TypeRegistry->GetOrMakeGeneratedType(
 			EWebAPISchemaType::Model,
diff --git a/Engine/Plugins/Experimental/Web/WebAPI/Source/WebAPIOpenAPI/Private/V3/WebAPIOpenAPIConverter.h b/Engine/Plugins/Experimental/Web/WebAPI/Source/WebAPIOpenAPI/Private/V3/WebAPIOpenAPIConverter.h
--- a/Engine/Plugins/Experimental/Web/WebAPI/Source/WebAPIOpenAPI/Private/V3/WebAPIOpenAPIConverter.h
+++ b/Engine/Plugins/Experimental/Web/WebAPI/Source/WebAPIOpenAPI/Private/V3/WebAPIOpenAPIConverter.h
@@ -65,17 +65,17 @@ namespace UE::WebAPI::OpenAPI
 		bool IsArray(const TSharedPtr<SchemaType>& InSchema);
 
 		template <typename SchemaType, typename ModelType>
-		static bool ConvertModelBase(const TSharedPtr<SchemaType>& InSchema, const TObjectPtr<ModelType>& OutModel);
+		static bool ConvertModelBase(const TSharedPtr<SchemaType>& InSchema, ModelType* OutModel);
 
 		template <typename SchemaType>
 		TObjectPtr<UWebAPIEnum> ConvertEnum(const TSharedPtr<SchemaType>& InSrcSchema, const FWebAPITypeNameVariant& InEnumTypeName = {}) const;
 
 		/** Modifies an existing Property using the supplied source object. */
 		template <typename SchemaType>
-		bool ConvertProperty(const FWebAPITypeNameVariant& InModelName, const FWebAPINameVariant& InPropertyName, const TSharedPtr<SchemaType>& InSchema, const FString& InDefinitionName, const TObjectPtr<UWebAPIProperty>& OutProperty);
+		bool ConvertProperty(const FWebAPITypeNameVariant& InModelName, const FWebAPINameVariant& InPropertyName, const TSharedPtr<SchemaType>& InSchema, const FString& InDefinitionName, UWebAPIProperty* OutProperty);
 
 		/** Modifies an existing Property using the supplied TypeInfo. */
-		bool ConvertProperty(const FWebAPITypeNameVariant& InModelName, const FWebAPINameVariant& InPropertyName, const FWebAPITypeNameVariant& InPropertyTypeName, const TObjectPtr<UWebAPIProperty>& OutProperty);
+		bool ConvertProperty(const FWebAPITypeNameVariant& InModelName, const FWebAPINameVariant& InPropertyName, const FWebAPITypeNameVariant& InPropertyTypeName, UWebAPIProperty* OutProperty);
 		
 		/** Creates a new Property from the supplied source object. */
 		template <typename SchemaType>
@@ -83,19 +83,19 @@ namespace UE::WebAPI::OpenAPI
 
 		/** Modifies an existing Model using the supplied source object. */
 		template <typename SchemaType>
-		bool ConvertModel(const TSharedPtr<SchemaType>& InSrcSchema, const FWebAPITypeNameVariant& InModelTypeName, const TObjectPtr<UWebAPIModel>& OutModel);
+		bool ConvertModel(const TSharedPtr<SchemaType>& InSrcSchema, const FWebAPITypeNameVariant& InModelTypeName, UWebAPIModel* OutModel);
 		
 		/** Creates a new Model from the supplied source object. */
 		template <typename SchemaType>
 		TObjectPtr<UWebAPIModel> ConvertModel(const TSharedPtr<SchemaType>& InSrcSchema, const FWebAPITypeNameVariant& InModelTypeName = {});
 
-		bool ConvertOperationParameter(const FWebAPINameVariant& InParameterName, const TSharedPtr<OpenAPI::V3::FParameterObject>& InParameter, const FString& InDefinitionName, const TObjectPtr<UWebAPIOperationParameter>& OutParameter);
+		bool ConvertOperationParameter(const FWebAPINameVariant& InParameterName, const TSharedPtr<OpenAPI::V3::FParameterObject>& InParameter, const FString& InDefinitionName, UWebAPIOperationParameter* OutParameter);
 
 		TObjectPtr<UWebAPIParameter> ConvertParameter(const TSharedPtr<OpenAPI::V3::FParameterObject>& InSrcParameter);
 		
-		bool ConvertRequest(const FWebAPITypeNameVariant& InOperationName, const TSharedPtr<OpenAPI::V3::FOperationObject>& InOperation, const TObjectPtr<UWebAPIOperationRequest>& OutRequest);
+		bool ConvertRequest(const FWebAPITypeNameVariant& InOperationName, const TSharedPtr<OpenAPI::V3::FOperationObject>& InOperation, UWebAPIOperationRequest* OutRequest);
 
-		bool ConvertResponse(const FWebAPITypeNameVariant& InOperationName, uint32 InResponseCode, const TSharedPtr<OpenAPI::V3::FResponseObject>& InResponse, const TObjectPtr<UWebAPIOperationResponse>& OutResponse);
+		bool ConvertResponse(const FWebAPITypeNameVariant& InOperationName, uint32 InResponseCode, const TSharedPtr<OpenAPI::V3::FResponseObject>& InResponse, UWebAPIOperationResponse* OutResponse);
 
 		TObjectPtr<UWebAPIOperation> ConvertOperation(const FString& InPath, const FString& InVerb, const TSharedPtr<OpenAPI::V3::FOperationObject>& InSrcOperation, const FWebAPITypeNameVariant& InOperationTypeName = {});
 
@@ -138,5 +138,5 @@ namespace UE::WebAPI::OpenAPI
 
 	/** Modifies an existing Model using the supplied source object. */
 	template <>
-	bool FWebAPIOpenAPISchemaConverter::ConvertModel(const TSharedPtr<OpenAPI::V3::FParameterObject>& InSrcSchema, const FWebAPITypeNameVariant& InModelTypeName, const TObjectPtr<UWebAPIModel>& OutModel);
+	bool FWebAPIOpenAPISchemaConverter::ConvertModel(const TSharedPtr<OpenAPI::V3::FParameterObject>& InSrcSchema, const FWebAPITypeNameVariant& InModelTypeName, UWebAPIModel* OutModel);
 };
diff --git a/Engine/Plugins/Experimental/Web/WebAPI/Source/WebAPIOpenAPI/Private/V3/WebAPIOpenAPIConverter.inl b/Engine/Plugins/Experimental/Web/WebAPI/Source/WebAPIOpenAPI/Private/V3/WebAPIOpenAPIConverter.inl
--- a/Engine/Plugins/Experimental/Web/WebAPI/Source/WebAPIOpenAPI/Private/V3/WebAPIOpenAPIConverter.inl
+++ b/Engine/Plugins/Experimental/Web/WebAPI/Source/WebAPIOpenAPI/Private/V3/WebAPIOpenAPIConverter.inl
@@ -12,7 +12,7 @@
 namespace UE::WebAPI::OpenAPI
 {
 	template <>
-	FString FWebAPIOpenAPISchemaConverter::GetDefaultJsonTypeForStorage<EWebAPIParameterStorage>(const EWebAPIParameterStorage& InStorage)
+	inline FString FWebAPIOpenAPISchemaConverter::GetDefaultJsonTypeForStorage<EWebAPIParameterStorage>(const EWebAPIParameterStorage& InStorage)
 	{
 		static TMap<EWebAPIParameterStorage, FString> StorageToJsonType = {
 			{ EWebAPIParameterStorage::Body, TEXT("object") },
@@ -26,7 +26,7 @@ namespace UE::WebAPI::OpenAPI
 	}
 
 	template <>
-	FString FWebAPIOpenAPISchemaConverter::GetDefaultJsonTypeForStorage<EWebAPIResponseStorage>(const EWebAPIResponseStorage& InStorage)
+	inline FString FWebAPIOpenAPISchemaConverter::GetDefaultJsonTypeForStorage<EWebAPIResponseStorage>(const EWebAPIResponseStorage& InStorage)
 	{
 		static TMap<EWebAPIResponseStorage, FString> StorageToJsonType = {
 			{ EWebAPIResponseStorage::Body, TEXT("object") },
@@ -123,7 +123,7 @@ namespace UE::WebAPI::OpenAPI
 	}
 
 	template <>
-	TSharedPtr<OpenAPI::V3::FSchemaObject> FWebAPIOpenAPISchemaConverter::ResolveReference<OpenAPI::V3::FSchemaObject>(const FString& InDefinitionName)
+	inline TSharedPtr<OpenAPI::V3::FSchemaObject> FWebAPIOpenAPISchemaConverter::ResolveReference<OpenAPI::V3::FSchemaObject>(const FString& InDefinitionName)
 	{
 		if(!InputSchema->Components.IsValid())
 		{
@@ -143,7 +143,7 @@ namespace UE::WebAPI::OpenAPI
 	}
 
 	template <>
-	TSharedPtr<OpenAPI::V3::FParameterObject> FWebAPIOpenAPISchemaConverter::ResolveReference<OpenAPI::V3::FParameterObject>(const FString& InDefinitionName)
+	inline TSharedPtr<OpenAPI::V3::FParameterObject> FWebAPIOpenAPISchemaConverter::ResolveReference<OpenAPI::V3::FParameterObject>(const FString& InDefinitionName)
 	{
 		if(!InputSchema->Components.IsValid())
 		{
@@ -163,7 +163,7 @@ namespace UE::WebAPI::OpenAPI
 	}
 
 	template <>
-	TSharedPtr<OpenAPI::V3::FResponseObject> FWebAPIOpenAPISchemaConverter::ResolveReference<V3::FResponseObject>(const FString& InDefinitionName)
+	inline TSharedPtr<OpenAPI::V3::FResponseObject> FWebAPIOpenAPISchemaConverter::ResolveReference<V3::FResponseObject>(const FString& InDefinitionName)
 	{
 		if(!InputSchema->Components.IsValid())
 		{
@@ -183,7 +183,7 @@ namespace UE::WebAPI::OpenAPI
 	}
 
 	template <>
-	TSharedPtr<OpenAPI::V3::FSecuritySchemeObject> FWebAPIOpenAPISchemaConverter::ResolveReference<V3::FSecuritySchemeObject>(const FString& InDefinitionName)
+	inline TSharedPtr<OpenAPI::V3::FSecuritySchemeObject> FWebAPIOpenAPISchemaConverter::ResolveReference<V3::FSecuritySchemeObject>(const FString& InDefinitionName)
 	{
 		if(!InputSchema->Components.IsValid())
 		{
@@ -203,7 +203,7 @@ namespace UE::WebAPI::OpenAPI
 	}
 
 	template <>
-	TSharedPtr<OpenAPI::V3::FRequestBodyObject> FWebAPIOpenAPISchemaConverter::ResolveReference<V3::FRequestBodyObject>(const FString& InDefinitionName)
+	inline TSharedPtr<OpenAPI::V3::FRequestBodyObject> FWebAPIOpenAPISchemaConverter::ResolveReference<V3::FRequestBodyObject>(const FString& InDefinitionName)
 	{
 		if (const Json::TJsonReference<OpenAPI::V3::FRequestBodyObject>* FoundDefinition = InputSchema->Components->RequestBodies.Find(InDefinitionName))
 		{
@@ -255,13 +255,13 @@ namespace UE::WebAPI::OpenAPI
 	}
 
 	template <>
-	bool FWebAPIOpenAPISchemaConverter::IsArray(const TSharedPtr<UE::WebAPI::OpenAPI::V3::FSchemaObject>& InSchema)
+	inline bool FWebAPIOpenAPISchemaConverter::IsArray(const TSharedPtr<UE::WebAPI::OpenAPI::V3::FSchemaObject>& InSchema)
 	{
 		return InSchema->Type.Get(TEXT(""))	== TEXT("array");
 	}
 
 	template <>
-	bool FWebAPIOpenAPISchemaConverter::IsArray(const TSharedPtr<UE::WebAPI::OpenAPI::V3::FParameterObject>& InSchema)
+	inline bool FWebAPIOpenAPISchemaConverter::IsArray(const TSharedPtr<UE::WebAPI::OpenAPI::V3::FParameterObject>& InSchema)
 	{
 		FString FoundDefinitionName;
 		return InSchema->Schema.IsSet()
@@ -271,7 +271,7 @@ namespace UE::WebAPI::OpenAPI
 
 	template <typename SchemaType, typename ModelType>
 	bool FWebAPIOpenAPISchemaConverter::ConvertModelBase(const TSharedPtr<SchemaType>& InSchema,
-														const TObjectPtr<ModelType>& OutModel)
+														ModelType* OutModel)
 	{
 		static_assert(TIsDerivedFrom<SchemaType, OpenAPI::V3::FSchemaObjectBase>::Value, "Type is not derived from OpenAPI::V3::FSchemaObjectBase.");
 
@@ -305,7 +305,7 @@ namespace UE::WebAPI::OpenAPI
 		DstEnum->Name = EnumTypeName;
 		DstEnum->Description = InSrcSchema->Description.Get(TEXT(""));
 
-		const TObjectPtr<UWebAPIModelBase> ModelBase = Cast<UWebAPIModelBase>(DstEnum);
+		UWebAPIModelBase* ModelBase = Cast<UWebAPIModelBase>(DstEnum);
 		if (!ConvertModelBase(InSrcSchema, ModelBase))
 		{
 			return nullptr;
@@ -324,25 +324,25 @@ namespace UE::WebAPI::OpenAPI
 	}
 
 	template <>
-	bool FWebAPIOpenAPISchemaConverter::ConvertProperty<OpenAPI::V3::FParameterObject>(
+	inline bool FWebAPIOpenAPISchemaConverter::ConvertProperty<OpenAPI::V3::FParameterObject>(
 		const FWebAPITypeNameVariant& InModelName,
 		const FWebAPINameVariant& InPropertyName,
 		const TSharedPtr<OpenAPI::V3::FParameterObject>& InParameter,
 		const FString& InDefinitionName,
-		const TObjectPtr<UWebAPIProperty>& OutProperty)
+		UWebAPIProperty* OutProperty)
 	{
 		const FString DefinitionName = !InDefinitionName.IsEmpty()
 										? InDefinitionName
 										: ProviderSettings.MakeNestedPropertyTypeName(InModelName, InPropertyName);
 
-		const TObjectPtr<UWebAPIModelBase> ModelBase = OutProperty;
+		UWebAPIModelBase* ModelBase = OutProperty;
 		if (!ConvertModelBase(InParameter, ModelBase))
 		{
 			return false;
 		}
 
 		OutProperty->bIsArray = IsArray(InParameter);
-		
+
 		OutProperty->Name = FWebAPINameInfo(NameTransformer(InPropertyName.ToString()), InPropertyName.GetJsonName(), OutProperty->Type);				
 		OutProperty->Name = ResolvePropertyName(OutProperty, InModelName, {});
 		OutProperty->bIsRequired = InParameter->bRequired.Get(false);
@@ -367,7 +367,7 @@ namespace UE::WebAPI::OpenAPI
 				}
 				EnumTypeName.TypeInfo->SetName(EnumName);
 
-				const TObjectPtr<UWebAPIEnum>& Enum = ConvertEnum(Schema, EnumTypeName);
+				const UWebAPIEnum* Enum = ConvertEnum(Schema, EnumTypeName);
 
 				const FText LogMessage = FText::FormatNamed(
 					LOCTEXT("AddedImplicitEnumForPropertyOfModel", "Implicit enum created for property \"{PropertyName}\" of model \"{ModelName}\"."),
@@ -379,7 +379,7 @@ namespace UE::WebAPI::OpenAPI
 				Enum->Name.TypeInfo->JsonName = InPropertyName.GetJsonName();
 
 				OutProperty->Type = Enum->Name;
-				OutProperty->Type.TypeInfo->Model = Enum;
+				OutProperty->Type.TypeInfo->SetModel(Enum);
 			}
 		}
 
@@ -391,7 +391,7 @@ namespace UE::WebAPI::OpenAPI
 														const FWebAPINameVariant& InPropertyName,
 														const TSharedPtr<SchemaType>& InSchema,
 														const FString& InDefinitionName,
-														const TObjectPtr<UWebAPIProperty>& OutProperty)
+														UWebAPIProperty* OutProperty)
 	{
 		static_assert(TIsDerivedFrom<SchemaType, OpenAPI::V3::FSchemaObjectBase>::Value, "Type is not derived from OpenAPI::V3::FSchemaObjectBase.");
 
@@ -416,7 +416,7 @@ namespace UE::WebAPI::OpenAPI
 		OutProperty->Type = ResolveType<SchemaType>(ItemSchema, DefinitionName);
 		check(OutProperty->Type.IsValid());
 
-		const TObjectPtr<UWebAPIModelBase> ModelBase = OutProperty;
+		UWebAPIModelBase* ModelBase = OutProperty;
 		if (!ConvertModelBase(ItemSchema, ModelBase))
 		{
 			return false;
@@ -435,7 +435,7 @@ namespace UE::WebAPI::OpenAPI
 			const FWebAPITypeNameVariant EnumTypeName = OutProperty->Type;
 			EnumTypeName.TypeInfo->SetName(ProviderSettings.MakeNestedPropertyTypeName(InModelName, InPropertyName));
 
-			const TObjectPtr<UWebAPIEnum>& Enum = ConvertEnum(InSchema, EnumTypeName);
+			const UWebAPIEnum* Enum = ConvertEnum(InSchema, EnumTypeName);
 
 			const FText LogMessage = FText::FormatNamed(
 				LOCTEXT("AddedImplicitEnumForPropertyOfModel", "Implicit enum created for property \"{PropertyName}\" of model \"{ModelName}\"."),
@@ -444,12 +444,12 @@ namespace UE::WebAPI::OpenAPI
 			MessageLog->LogInfo(LogMessage, FWebAPIOpenAPIProvider::LogName);
 
 			Enum->Name.TypeInfo->DebugString += LogMessage.ToString();
-			
+
 			Enum->Name.TypeInfo->JsonName = InPropertyName.GetJsonName();
 			Enum->Name.TypeInfo->SetNested(InModelName);
 
 			OutProperty->Type = Enum->Name;
-			OutProperty->Type.TypeInfo->Model = Enum;
+			OutProperty->Type.TypeInfo->SetModel(Enum);
 		}
 		// Add struct as it's own model, and reference it as this properties type
 		else if (OutProperty->Type.ToString(true).IsEmpty())
@@ -461,7 +461,7 @@ namespace UE::WebAPI::OpenAPI
 			ModelTypeName.TypeInfo->Prefix = TEXT("F");
 			ModelTypeName.TypeInfo->SetNested(InModelName);
 
-			const TObjectPtr<UWebAPIModel>& Model = OutputSchema->AddModel<UWebAPIModel>(ModelTypeName.HasTypeInfo() ? ModelTypeName.TypeInfo.Get() : nullptr);
+			UWebAPIModel* Model = OutputSchema->AddModel<UWebAPIModel>(ModelTypeName.HasTypeInfo() ? ModelTypeName.TypeInfo.Get() : nullptr);
 			ConvertModel(InSchema, {}, Model);
 
 			const FText LogMessage = FText::FormatNamed(
@@ -476,7 +476,7 @@ namespace UE::WebAPI::OpenAPI
 			Model->Name.TypeInfo->JsonType = UWebAPIStaticTypeRegistry::ToFromJsonType;
 
 			OutProperty->Type = Model->Name;
-			OutProperty->Type.TypeInfo->Model = Model;
+			OutProperty->Type.TypeInfo->SetModel(Model);
 		}
 		
 		OutProperty->Name = ResolvePropertyName(OutProperty, InModelName);
@@ -484,7 +484,7 @@ namespace UE::WebAPI::OpenAPI
 		return true;
 	}
 
-		template <typename SchemaType>
+	template <typename SchemaType>
 	TObjectPtr<UWebAPIProperty> FWebAPIOpenAPISchemaConverter::ConvertProperty(const TSharedPtr<SchemaType>& InSrcSchema, const TObjectPtr<UWebAPIModel>& InModel, const FWebAPINameVariant& InPropertyName, const FString& InDefinitionName)
 	{
 		static_assert(TIsDerivedFrom<SchemaType, OpenAPI::V3::FSchemaObjectBase>::Value, "Type is not derived from OpenAPI::V3::FSchemaObjectBase.");
@@ -509,7 +509,7 @@ namespace UE::WebAPI::OpenAPI
 	}
 
 	template <typename SchemaType>
-	bool FWebAPIOpenAPISchemaConverter::ConvertModel(const TSharedPtr<SchemaType>& InSrcSchema, const FWebAPITypeNameVariant& InModelTypeName, const TObjectPtr<UWebAPIModel>& OutModel)
+	bool FWebAPIOpenAPISchemaConverter::ConvertModel(const TSharedPtr<SchemaType>& InSrcSchema, const FWebAPITypeNameVariant& InModelTypeName, UWebAPIModel* OutModel)
 	{
 		static_assert(TIsDerivedFrom<SchemaType, OpenAPI::V3::FSchemaObjectBase>::Value, "Type is not derived from OpenAPI::V3::FSchemaObjectBase.");
 
@@ -548,7 +548,7 @@ namespace UE::WebAPI::OpenAPI
 
 		const TObjectPtr<UWebAPIModel>& Model = OutModel ? OutModel : OutputSchema->AddModel<UWebAPIModel>(ModelTypeName.TypeInfo.Get());
 
-		const TObjectPtr<UWebAPIModelBase> ModelBase = Model;
+		UWebAPIModelBase* ModelBase = Model;
 		if (!ConvertModelBase(InSrcSchema, ModelBase))
 		{
 			return false;
@@ -580,7 +580,7 @@ namespace UE::WebAPI::OpenAPI
 	}
 
 	template <>
-	bool FWebAPIOpenAPISchemaConverter::ConvertModel<OpenAPI::V3::FParameterObject>(const TSharedPtr<OpenAPI::V3::FParameterObject>& InSrcSchema, const FWebAPITypeNameVariant& InModelTypeName, const TObjectPtr<UWebAPIModel>& OutModel)
+	inline bool FWebAPIOpenAPISchemaConverter::ConvertModel<OpenAPI::V3::FParameterObject>(const TSharedPtr<OpenAPI::V3::FParameterObject>& InSrcSchema, const FWebAPITypeNameVariant& InModelTypeName, UWebAPIModel* OutModel)
 	{
 		FWebAPITypeNameVariant ModelTypeName;
 		if(InModelTypeName.IsValid() && InModelTypeName.HasTypeInfo())
@@ -615,9 +615,9 @@ namespace UE::WebAPI::OpenAPI
 			return false;
 		}
 
-		const TObjectPtr<UWebAPIModel>& Model = OutModel ? OutModel : OutputSchema->AddModel<UWebAPIModel>(ModelTypeName.TypeInfo.Get());
+		auto* Model = OutModel ? OutModel : OutputSchema->AddModel<UWebAPIModel>(ModelTypeName.TypeInfo.Get());
 
-		const TObjectPtr<UWebAPIModelBase> ModelBase = Model;
+		UWebAPIModelBase* ModelBase = Model;
 		if (!ConvertModelBase(InSrcSchema, ModelBase))
 		{
 			return false;
@@ -652,7 +652,7 @@ namespace UE::WebAPI::OpenAPI
 				TEXT("F"));
 		}
 		
-		const TObjectPtr<UWebAPIModel>& Model = OutputSchema->AddModel<UWebAPIModel>(ModelTypeName.TypeInfo.Get());
+		UWebAPIModel* Model = OutputSchema->AddModel<UWebAPIModel>(ModelTypeName.TypeInfo.Get());
 		if(ConvertModel(InSrcSchema, ModelTypeName, Model))
 		{
 			return Model;
diff --git a/Engine/Plugins/Experimental/Web/WebAPI/Source/WebAPIOpenAPI/WebAPIOpenAPI.build.cs b/Engine/Plugins/Experimental/Web/WebAPI/Source/WebAPIOpenAPI/WebAPIOpenAPI.build.cs
--- a/Engine/Plugins/Experimental/Web/WebAPI/Source/WebAPIOpenAPI/WebAPIOpenAPI.build.cs
+++ b/Engine/Plugins/Experimental/Web/WebAPI/Source/WebAPIOpenAPI/WebAPIOpenAPI.build.cs
@@ -36,5 +36,13 @@ public class WebAPIOpenAPI : ModuleRules
 				"WebAPI",
 				"WebAPIEditor"
 			});
+
+		if (Target.bBuildEditor)
+		{
+			PrivateDependencyModuleNames.AddRange(
+				new string[] {
+					"Projects"
+				});
+		}
     }
 }
diff --git a/Engine/Source/Runtime/Slate/Private/Framework/MultiBox/SClippingHorizontalBox.cpp b/Engine/Source/Runtime/Slate/Private/Framework/MultiBox/SClippingHorizontalBox.cpp
--- a/Engine/Source/Runtime/Slate/Private/Framework/MultiBox/SClippingHorizontalBox.cpp
+++ b/Engine/Source/Runtime/Slate/Private/Framework/MultiBox/SClippingHorizontalBox.cpp
@@ -1,13 +1,11 @@
 // Copyright Epic Games, Inc. All Rights Reserved.
 
 #include "Framework/MultiBox/SClippingHorizontalBox.h"
+#include "Framework/Application/SlateApplication.h"
 #include "Layout/ArrangedChildren.h"
-#include "Rendering/DrawElements.h"
+#include "Styling/ToolBarStyle.h"
 #include "Widgets/Images/SImage.h"
 #include "Widgets/Input/SComboButton.h"
-#include "Framework/Application/SlateApplication.h"
-#include "Styling/ToolBarStyle.h"
-
 
 void SClippingHorizontalBox::OnArrangeChildren( const FGeometry& AllottedGeometry, FArrangedChildren& ArrangedChildren ) const
 {
@@ -129,7 +127,7 @@ void SClippingHorizontalBox::AddWrapButton()
 	// Add the wrap button
 	AddSlot()
 	.FillWidth(0.0f) // Effectively makes this widget 0 width, so it exists as a slot/child, but isn't considered for layout
-	.Padding( 0.f )
+	.Padding(0.f)
 	[
 		WrapButton.ToSharedRef()
 	];
