Fixed const mismatch.

diff --git a/Engine/Plugins/Experimental/Web/WebAPI/Source/WebAPIEditor/Public/Dom/WebAPIType.h b/Engine/Plugins/Experimental/Web/WebAPI/Source/WebAPIEditor/Public/Dom/WebAPIType.h
--- a/Engine/Plugins/Experimental/Web/WebAPI/Source/WebAPIEditor/Public/Dom/WebAPIType.h
+++ b/Engine/Plugins/Experimental/Web/WebAPI/Source/WebAPIEditor/Public/Dom/WebAPIType.h
@@ -1,4 +1,4 @@
-﻿// Copyright Epic Games, Inc. All Rights Reserved.
+// Copyright Epic Games, Inc. All Rights Reserved.
 
 #pragma once
 
@@ -164,7 +164,7 @@ public:
 			&& std::is_base_of_v<IWebAPISchemaObjectInterface, std::decay_t<ModelType>>)>
 	void SetModel(const ModelType* InModel)
 	{
-		Model = InModel;
+		Model = static_cast<const UObject*>(const_cast<UObject*>(InModel));
 	}
 
 #if WITH_EDITOR
