CIS fix.

diff --git a/Engine/Plugins/Experimental/Web/WebAPI/Source/WebAPIEditor/Public/Dom/WebAPIType.h b/Engine/Plugins/Experimental/Web/WebAPI/Source/WebAPIEditor/Public/Dom/WebAPIType.h
--- a/Engine/Plugins/Experimental/Web/WebAPI/Source/WebAPIEditor/Public/Dom/WebAPIType.h
+++ b/Engine/Plugins/Experimental/Web/WebAPI/Source/WebAPIEditor/Public/Dom/WebAPIType.h
@@ -164,7 +164,7 @@ public:
 			&& std::is_base_of_v<IWebAPISchemaObjectInterface, std::decay_t<ModelType>>)>
 	void SetModel(const ModelType* InModel)
 	{
-		Model = static_cast<const UObject*>(const_cast<UObject*>(InModel));
+		Model = const_cast<UObject*>(static_cast<const UObject*>(InModel));
 	}
 
 #if WITH_EDITOR
