[WebAPI] cis fix

#rnx

diff --git a/Engine/Plugins/Experimental/Web/WebAPI/Source/WebAPIOpenAPI/Private/V2/WebAPISwaggerConverter.inl b/Engine/Plugins/Experimental/Web/WebAPI/Source/WebAPIOpenAPI/Private/V2/WebAPISwaggerConverter.inl
--- a/Engine/Plugins/Experimental/Web/WebAPI/Source/WebAPIOpenAPI/Private/V2/WebAPISwaggerConverter.inl
+++ b/Engine/Plugins/Experimental/Web/WebAPI/Source/WebAPIOpenAPI/Private/V2/WebAPISwaggerConverter.inl
@@ -136,7 +136,7 @@ namespace UE::WebAPI::Swagger
 			{
 				ModelTypeName = InModelTypeName;
 			}
-			else if(OutModel->Name.HasTypeInfo())
+			else if(OutModel && OutModel->Name.HasTypeInfo())
 			{
 				ModelTypeName = OutModel->Name;
 			}
@@ -207,7 +207,7 @@ namespace UE::WebAPI::Swagger
 			{
 				ModelTypeName = InModelTypeName;
 			}
-			else if(OutModel->Name.HasTypeInfo())
+			else if(OutModel && OutModel->Name.HasTypeInfo())
 			{
 				ModelTypeName = OutModel->Name;
 			}
diff --git a/Engine/Plugins/Experimental/Web/WebAPI/Source/WebAPIOpenAPI/Private/V3/WebAPIOpenAPIConverter.inl b/Engine/Plugins/Experimental/Web/WebAPI/Source/WebAPIOpenAPI/Private/V3/WebAPIOpenAPIConverter.inl
--- a/Engine/Plugins/Experimental/Web/WebAPI/Source/WebAPIOpenAPI/Private/V3/WebAPIOpenAPIConverter.inl
+++ b/Engine/Plugins/Experimental/Web/WebAPI/Source/WebAPIOpenAPI/Private/V3/WebAPIOpenAPIConverter.inl
@@ -518,7 +518,7 @@ namespace UE::WebAPI::OpenAPI
 		{
 			ModelTypeName = InModelTypeName;
 		}
-		else if(OutModel->Name.HasTypeInfo())
+		else if(OutModel && OutModel->Name.HasTypeInfo())
 		{
 			ModelTypeName = OutModel->Name;
 		}
