[WebAPI] cis fix

#rnx

diff --git a/Engine/Plugins/Experimental/Web/WebAPI/Source/WebAPIOpenAPI/Private/V3/WebAPIOpenAPIConverter.inl b/Engine/Plugins/Experimental/Web/WebAPI/Source/WebAPIOpenAPI/Private/V3/WebAPIOpenAPIConverter.inl
--- a/Engine/Plugins/Experimental/Web/WebAPI/Source/WebAPIOpenAPI/Private/V3/WebAPIOpenAPIConverter.inl
+++ b/Engine/Plugins/Experimental/Web/WebAPI/Source/WebAPIOpenAPI/Private/V3/WebAPIOpenAPIConverter.inl
@@ -587,7 +587,7 @@ namespace UE::WebAPI::OpenAPI
 		{
 			ModelTypeName = InModelTypeName;
 		}
-		else if(OutModel->Name.HasTypeInfo())
+		else if(OutModel && OutModel->Name.HasTypeInfo())
 		{
 			ModelTypeName = OutModel->Name;
 		}
@@ -615,7 +615,7 @@ namespace UE::WebAPI::OpenAPI
 			return false;
 		}
 
-		auto* Model = OutModel ? OutModel : OutputSchema->AddModel<UWebAPIModel>(ModelTypeName.TypeInfo.Get());
+		UWebAPIModel* Model = OutModel ? OutModel : OutputSchema->AddModel<UWebAPIModel>(ModelTypeName.TypeInfo.Get());
 
 		UWebAPIModelBase* ModelBase = Model;
 		if (!ConvertModelBase(InSrcSchema, ModelBase))
