Removed extraneous GetSelfPin implementations.

Since GetThenPin was added to UK2Node in 5.x, derived classes don't need to provide their own implementations. These were largely identical to the one in UK2Node. Additionally, overriding non-virtual functions can be error prone if there's an expectation of polymorphic behavior.

#rnx

diff --git a/Engine/Plugins/Animation/LiveLink/Source/LiveLinkGraphNode/Private/K2Node_EvaluateLiveLinkFrame.cpp b/Engine/Plugins/Animation/LiveLink/Source/LiveLinkGraphNode/Private/K2Node_EvaluateLiveLinkFrame.cpp
--- a/Engine/Plugins/Animation/LiveLink/Source/LiveLinkGraphNode/Private/K2Node_EvaluateLiveLinkFrame.cpp
+++ b/Engine/Plugins/Animation/LiveLink/Source/LiveLinkGraphNode/Private/K2Node_EvaluateLiveLinkFrame.cpp
@@ -204,15 +204,6 @@ FText UK2Node_EvaluateLiveLinkFrame::GetTooltipText() const
 	return LOCTEXT("NodeTooltip", "Attempts to Get a LiveLink Frame from a subject using a given Role");
 }
 
-UEdGraphPin* UK2Node_EvaluateLiveLinkFrame::GetThenPin()const
-{
-	const UEdGraphSchema_K2* K2Schema = GetDefault<UEdGraphSchema_K2>();
-
-	UEdGraphPin* Pin = FindPinChecked(UEdGraphSchema_K2::PN_Then);
-	check(Pin->Direction == EGPD_Output);
-	return Pin;
-}
-
 UEdGraphPin* UK2Node_EvaluateLiveLinkFrame::GetLiveLinkRolePin() const
 {
 	UEdGraphPin* Pin = FindPinChecked(UK2Node_EvaluateLiveLinkFrameHelper::LiveLinkRolePinName);
diff --git a/Engine/Plugins/Animation/LiveLink/Source/LiveLinkGraphNode/Private/K2Node_UpdateVirtualSubjectDataBase.cpp b/Engine/Plugins/Animation/LiveLink/Source/LiveLinkGraphNode/Private/K2Node_UpdateVirtualSubjectDataBase.cpp
--- a/Engine/Plugins/Animation/LiveLink/Source/LiveLinkGraphNode/Private/K2Node_UpdateVirtualSubjectDataBase.cpp
+++ b/Engine/Plugins/Animation/LiveLink/Source/LiveLinkGraphNode/Private/K2Node_UpdateVirtualSubjectDataBase.cpp
@@ -161,15 +161,6 @@ void UK2Node_UpdateVirtualSubjectDataBase::EarlyValidation(class FCompilerResult
 	}
 }
 
-UEdGraphPin* UK2Node_UpdateVirtualSubjectDataBase::GetThenPin() const
-{
-	const UEdGraphSchema_K2* K2Schema = GetDefault<UEdGraphSchema_K2>();
-
-	UEdGraphPin* Pin = FindPinChecked(UEdGraphSchema_K2::PN_Then);
-	check(Pin->Direction == EGPD_Output);
-	return Pin;
-}
-
 UEdGraphPin* UK2Node_UpdateVirtualSubjectDataBase::GetLiveLinkStructPin() const
 {
 	UEdGraphPin* Pin = FindPinChecked(LiveLinkStructPinName);
diff --git a/Engine/Plugins/Animation/LiveLink/Source/LiveLinkGraphNode/Public/K2Node_EvaluateLiveLinkFrame.h b/Engine/Plugins/Animation/LiveLink/Source/LiveLinkGraphNode/Public/K2Node_EvaluateLiveLinkFrame.h
--- a/Engine/Plugins/Animation/LiveLink/Source/LiveLinkGraphNode/Public/K2Node_EvaluateLiveLinkFrame.h
+++ b/Engine/Plugins/Animation/LiveLink/Source/LiveLinkGraphNode/Public/K2Node_EvaluateLiveLinkFrame.h
@@ -45,9 +45,6 @@ public:
 	/** Get the return types of our struct */
 	UScriptStruct* GetReturnTypeForOutputDataStruct() const;
 
-	/** Get the then output pin */
-	UEdGraphPin* GetThenPin() const;
-
 	/** Get the Live Link Role input pin */
 	UEdGraphPin* GetLiveLinkRolePin() const;
 
diff --git a/Engine/Plugins/Animation/LiveLink/Source/LiveLinkGraphNode/Public/K2Node_UpdateVirtualSubjectDataBase.h b/Engine/Plugins/Animation/LiveLink/Source/LiveLinkGraphNode/Public/K2Node_UpdateVirtualSubjectDataBase.h
--- a/Engine/Plugins/Animation/LiveLink/Source/LiveLinkGraphNode/Public/K2Node_UpdateVirtualSubjectDataBase.h
+++ b/Engine/Plugins/Animation/LiveLink/Source/LiveLinkGraphNode/Public/K2Node_UpdateVirtualSubjectDataBase.h
@@ -38,10 +38,6 @@ public:
 	virtual void EarlyValidation(class FCompilerResultsLog& MessageLog) const override;
 	//~ End UK2Node Interface
 
-
-	/** Get the then output pin */
-	UEdGraphPin* GetThenPin() const;
-
 	/** Get the Live Link Struct input pin */
 	UEdGraphPin* GetLiveLinkStructPin() const;
 
diff --git a/Engine/Plugins/Experimental/Web/WebAPI/Source/WebAPIBlueprintGraph/Private/K2Node_WebAPIOperation.cpp b/Engine/Plugins/Experimental/Web/WebAPI/Source/WebAPIBlueprintGraph/Private/K2Node_WebAPIOperation.cpp
--- a/Engine/Plugins/Experimental/Web/WebAPI/Source/WebAPIBlueprintGraph/Private/K2Node_WebAPIOperation.cpp
+++ b/Engine/Plugins/Experimental/Web/WebAPI/Source/WebAPIBlueprintGraph/Private/K2Node_WebAPIOperation.cpp
@@ -158,7 +158,7 @@ void UK2Node_WebAPIOperation::AllocateDefaultPins()
 		CreatePin(EGPD_Input, UEdGraphSchema_K2::PC_Exec, UEdGraphSchema_K2::PN_Execute);
 	}
 
-	if(!GetThenPin())
+	if(!FindThenPin())
 	{
 		CreatePin(EGPD_Output, UEdGraphSchema_K2::PC_Exec, UEdGraphSchema_K2::PN_Then);
 	}
@@ -911,7 +911,7 @@ TArray<UEdGraphPin*> UK2Node_WebAPIOperation::GetResponsePins() const
 	return UE::WebAPI::Graph::GetResponsePins(this);
 }
 
-UEdGraphPin* UK2Node_WebAPIOperation::GetThenPin() const
+UEdGraphPin* UK2Node_WebAPIOperation::FindThenPin() const
 {
 	return FindPin(UEdGraphSchema_K2::PN_Then, EEdGraphPinDirection::EGPD_Output, UEdGraphSchema_K2::PC_Exec);
 }
@@ -981,7 +981,7 @@ public:
 		ExecIn = InNode->GetExecPin();
 		RequestParameters = InNode->GetRequestPins();
 
-		ExecOut = InNode->GetThenPin();
+		ExecOut = InNode->FindThenPin();
 		ExecPositive = InNode->FindPin(UE::WebAPI::Operation::PositiveOutcomeName, EEdGraphPinDirection::EGPD_Output, UEdGraphSchema_K2::PC_Exec, true);
 		ExecNegative = InNode->FindPin(UE::WebAPI::Operation::NegativeOutcomeName, EEdGraphPinDirection::EGPD_Output, UEdGraphSchema_K2::PC_Exec, true);
 		Responses = InNode->GetResponsePins();
@@ -1082,7 +1082,7 @@ public:
 		check(InNode->PositiveDelegateProperty.IsValid());
 		DelegateProperty = InNode->PositiveDelegateProperty.Get(); // Choose one, doesn't matter
 		
-		ExecOut = InNode->GetThenPin();
+		ExecOut = InNode->FindThenPin();
 		
 		PositiveEventNodeMap = FOutcomeEventPinMap(UE::WebAPI::Operation::PositiveOutcomeName, InPositiveEvent, DelegateProperty->SignatureFunction);
 		NegativeEventNodeMap = FOutcomeEventPinMap(UE::WebAPI::Operation::NegativeOutcomeName, InNegativeEvent, DelegateProperty->SignatureFunction);
diff --git a/Engine/Plugins/Experimental/Web/WebAPI/Source/WebAPIBlueprintGraph/Private/K2Node_WebAPIOperation.h b/Engine/Plugins/Experimental/Web/WebAPI/Source/WebAPIBlueprintGraph/Private/K2Node_WebAPIOperation.h
--- a/Engine/Plugins/Experimental/Web/WebAPI/Source/WebAPIBlueprintGraph/Private/K2Node_WebAPIOperation.h
+++ b/Engine/Plugins/Experimental/Web/WebAPI/Source/WebAPIBlueprintGraph/Private/K2Node_WebAPIOperation.h
@@ -124,6 +124,8 @@ private:
 		const EEdGraphPinDirection& InDirection = EEdGraphPinDirection::EGPD_MAX,
 		bool bOnlySplitPins = true) const;
 
+	UEdGraphPin* FindThenPin() const;
+
 	/** Get all input pins that correspond to operation parameters. */
 	TArray<UEdGraphPin*> GetRequestPins() const;
 
@@ -133,9 +135,6 @@ private:
 	/** Get all (non-exec, non-error) output pins, will return empty if this uses Callbacks. */
 	TArray<UEdGraphPin*> GetResponsePins() const;
 
-	/** Get the output execution pin. */
-	UEdGraphPin* GetThenPin() const;
-
 	/** Get both outcome exec pins, will return empty if this uses Callbacks. */
 	TArray<UEdGraphPin*> GetResponseExecPins() const;
 
diff --git a/Engine/Plugins/Runtime/ModelViewViewModel/Source/ModelViewViewModelBlueprint/Private/Node/MVVMK2Node_AreSourcesValidForEvent.cpp b/Engine/Plugins/Runtime/ModelViewViewModel/Source/ModelViewViewModelBlueprint/Private/Node/MVVMK2Node_AreSourcesValidForEvent.cpp
--- a/Engine/Plugins/Runtime/ModelViewViewModel/Source/ModelViewViewModelBlueprint/Private/Node/MVVMK2Node_AreSourcesValidForEvent.cpp
+++ b/Engine/Plugins/Runtime/ModelViewViewModel/Source/ModelViewViewModelBlueprint/Private/Node/MVVMK2Node_AreSourcesValidForEvent.cpp
@@ -100,11 +100,6 @@ void UMVVMK2Node_AreSourcesValidForEvent::ExpandNode(FKismetCompilerContext& Com
 	BreakAllNodeLinks();
 }
 
-UEdGraphPin* UMVVMK2Node_AreSourcesValidForEvent::GetThenPin() const
-{
-	return FindPinChecked(UEdGraphSchema_K2::PN_Then);
-}
-
 UEdGraphPin* UMVVMK2Node_AreSourcesValidForEvent::GetElsePin() const
 {
 	return FindPinChecked(UEdGraphSchema_K2::PN_Else);
diff --git a/Engine/Plugins/Runtime/ModelViewViewModel/Source/ModelViewViewModelBlueprint/Private/Node/MVVMK2Node_AreSourcesValidForEvent.h b/Engine/Plugins/Runtime/ModelViewViewModel/Source/ModelViewViewModelBlueprint/Private/Node/MVVMK2Node_AreSourcesValidForEvent.h
--- a/Engine/Plugins/Runtime/ModelViewViewModel/Source/ModelViewViewModelBlueprint/Private/Node/MVVMK2Node_AreSourcesValidForEvent.h
+++ b/Engine/Plugins/Runtime/ModelViewViewModel/Source/ModelViewViewModelBlueprint/Private/Node/MVVMK2Node_AreSourcesValidForEvent.h
@@ -22,8 +22,6 @@ public:
 	virtual void AllocateDefaultPins() override;
 	//~ End K2Node Interface.
 
-	/** Get the then output pin */
-	UEdGraphPin* GetThenPin() const;
 	/** Get the else output pin */
 	UEdGraphPin* GetElsePin() const;
 
diff --git a/Engine/Plugins/VirtualProduction/DMX/DMXEngine/Source/DMXBlueprintGraph/Private/K2Node_GetDMXAttributeValues.cpp b/Engine/Plugins/VirtualProduction/DMX/DMXEngine/Source/DMXBlueprintGraph/Private/K2Node_GetDMXAttributeValues.cpp
--- a/Engine/Plugins/VirtualProduction/DMX/DMXEngine/Source/DMXBlueprintGraph/Private/K2Node_GetDMXAttributeValues.cpp
+++ b/Engine/Plugins/VirtualProduction/DMX/DMXEngine/Source/DMXBlueprintGraph/Private/K2Node_GetDMXAttributeValues.cpp
@@ -409,21 +409,6 @@ UEdGraphPin* UK2Node_GetDMXAttributeValues::GetOutputIsSuccessPin() const
 	return Pin;
 }
 
-UEdGraphPin* UK2Node_GetDMXAttributeValues::GetThenPin() const
-{
-	const UEdGraphSchema_K2* K2Schema = GetDefault<UEdGraphSchema_K2>();
-
-	UEdGraphPin* Pin = FindPin(UEdGraphSchema_K2::PN_Then);
-	if (Pin == nullptr)
-	{
-		UE_LOG_DMXBLUEPRINTGRAPH(Error, TEXT("No ThenPin found"));
-		return nullptr;
-	}
-
-	check(Pin->Direction == EGPD_Output);
-	return Pin;
-}
-
 void UK2Node_GetDMXAttributeValues::ExposeAttributes()
 {
 	if (bIsExposed == true && UserDefinedPins.Num())
diff --git a/Engine/Plugins/VirtualProduction/DMX/DMXEngine/Source/DMXBlueprintGraph/Public/K2Node_GetDMXAttributeValues.h b/Engine/Plugins/VirtualProduction/DMX/DMXEngine/Source/DMXBlueprintGraph/Public/K2Node_GetDMXAttributeValues.h
--- a/Engine/Plugins/VirtualProduction/DMX/DMXEngine/Source/DMXBlueprintGraph/Public/K2Node_GetDMXAttributeValues.h
+++ b/Engine/Plugins/VirtualProduction/DMX/DMXEngine/Source/DMXBlueprintGraph/Public/K2Node_GetDMXAttributeValues.h
@@ -74,7 +74,6 @@ public:
 	UEdGraphPin* GetInputDMXFixturePatchPin() const;
 	UEdGraphPin* GetOutputAttributesMapPin() const;
 	UEdGraphPin* GetOutputIsSuccessPin() const;
-	UEdGraphPin* GetThenPin() const;
 
 	/** Expose DMX function pins from fixture patch active mode */
 	void ExposeAttributes();
diff --git a/Engine/Plugins/Web/HttpBlueprint/Source/HttpBlueprintGraph/Private/K2Node_HttpRequest.cpp b/Engine/Plugins/Web/HttpBlueprint/Source/HttpBlueprintGraph/Private/K2Node_HttpRequest.cpp
--- a/Engine/Plugins/Web/HttpBlueprint/Source/HttpBlueprintGraph/Private/K2Node_HttpRequest.cpp
+++ b/Engine/Plugins/Web/HttpBlueprint/Source/HttpBlueprintGraph/Private/K2Node_HttpRequest.cpp
@@ -472,11 +472,4 @@ UEdGraphPin* UK2Node_HttpRequest::GetErrorPin() const
 	return Pin;
 }
 
-UEdGraphPin* UK2Node_HttpRequest::GetThenPin() const
-{
-	UEdGraphPin* Pin = FindPinChecked(UEdGraphSchema_K2::PN_Then);
-	check(Pin->Direction == EGPD_Output);
-	return Pin;
-}
-
 #undef LOCTEXT_NAMESPACE
diff --git a/Engine/Plugins/Web/HttpBlueprint/Source/HttpBlueprintGraph/Private/K2Node_HttpRequest.h b/Engine/Plugins/Web/HttpBlueprint/Source/HttpBlueprintGraph/Private/K2Node_HttpRequest.h
--- a/Engine/Plugins/Web/HttpBlueprint/Source/HttpBlueprintGraph/Private/K2Node_HttpRequest.h
+++ b/Engine/Plugins/Web/HttpBlueprint/Source/HttpBlueprintGraph/Private/K2Node_HttpRequest.h
@@ -52,7 +52,6 @@ protected:
 	UEdGraphPin* GetUrlPin() const;
 	UEdGraphPin* GetSuccessPin() const;
 	UEdGraphPin* GetErrorPin() const;
-	UEdGraphPin* GetThenPin() const;
 	
 private:
 	UPROPERTY()
diff --git a/Engine/Source/Editor/BlueprintGraph/Classes/K2Node_AssignmentStatement.h b/Engine/Source/Editor/BlueprintGraph/Classes/K2Node_AssignmentStatement.h
--- a/Engine/Source/Editor/BlueprintGraph/Classes/K2Node_AssignmentStatement.h
+++ b/Engine/Source/Editor/BlueprintGraph/Classes/K2Node_AssignmentStatement.h
@@ -48,8 +48,6 @@ class UK2Node_AssignmentStatement : public UK2Node
 	virtual int32 GetNodeRefreshPriority() const override { return EBaseNodeRefreshPriority::Low_UsesDependentWildcard; }
 	//~ End UK2Node Interface
 
-	/** Get the Then output pin */
-	BLUEPRINTGRAPH_API UEdGraphPin* GetThenPin() const;
 	/** Get the Variable input pin */
 	BLUEPRINTGRAPH_API UEdGraphPin* GetVariablePin() const;
 	/** Get the Value input pin */
diff --git a/Engine/Source/Editor/BlueprintGraph/Classes/K2Node_CallFunction.h b/Engine/Source/Editor/BlueprintGraph/Classes/K2Node_CallFunction.h
--- a/Engine/Source/Editor/BlueprintGraph/Classes/K2Node_CallFunction.h
+++ b/Engine/Source/Editor/BlueprintGraph/Classes/K2Node_CallFunction.h
@@ -154,8 +154,6 @@ public:
 	/** Returns the UFunction that this class is pointing to */
 	UFunction* GetTargetFunction() const;
 
-	/** Get the then output pin */
-	UEdGraphPin* GetThenPin() const;
 	/** Get the return value pin */
 	UEdGraphPin* GetReturnValuePin() const;
 
diff --git a/Engine/Source/Editor/BlueprintGraph/Classes/K2Node_ConstructObjectFromClass.h b/Engine/Source/Editor/BlueprintGraph/Classes/K2Node_ConstructObjectFromClass.h
--- a/Engine/Source/Editor/BlueprintGraph/Classes/K2Node_ConstructObjectFromClass.h
+++ b/Engine/Source/Editor/BlueprintGraph/Classes/K2Node_ConstructObjectFromClass.h
@@ -58,8 +58,6 @@ class BLUEPRINTGRAPH_API UK2Node_ConstructObjectFromClass : public UK2Node
 	/** See if this is a spawn variable pin, or a 'default' pin */
 	virtual bool IsSpawnVarPin(UEdGraphPin* Pin) const;
 
-	/** Get the then output pin */
-	UEdGraphPin* GetThenPin() const;
 	/** Get the blueprint input pin */	
 	UEdGraphPin* GetClassPin(const TArray<UEdGraphPin*>* InPinsToSearch=NULL) const;
 	/** Get the world context input pin, can return NULL */
diff --git a/Engine/Source/Editor/BlueprintGraph/Classes/K2Node_GetDataTableRow.h b/Engine/Source/Editor/BlueprintGraph/Classes/K2Node_GetDataTableRow.h
--- a/Engine/Source/Editor/BlueprintGraph/Classes/K2Node_GetDataTableRow.h
+++ b/Engine/Source/Editor/BlueprintGraph/Classes/K2Node_GetDataTableRow.h
@@ -54,8 +54,6 @@ class BLUEPRINTGRAPH_API UK2Node_GetDataTableRow : public UK2Node
 	/** Get the return type of our struct */
 	UScriptStruct* GetReturnTypeForStruct();
 
-	/** Get the then output pin */
-	UEdGraphPin* GetThenPin() const;
 	/** Get the Data Table input pin */
 	UEdGraphPin* GetDataTablePin(const TArray<UEdGraphPin*>* InPinsToSearch=NULL) const;
 	/** Get the spawn transform input pin */	
diff --git a/Engine/Source/Editor/BlueprintGraph/Classes/K2Node_IfThenElse.h b/Engine/Source/Editor/BlueprintGraph/Classes/K2Node_IfThenElse.h
--- a/Engine/Source/Editor/BlueprintGraph/Classes/K2Node_IfThenElse.h
+++ b/Engine/Source/Editor/BlueprintGraph/Classes/K2Node_IfThenElse.h
@@ -38,8 +38,6 @@ class UK2Node_IfThenElse : public UK2Node
 	virtual FText GetMenuCategory() const override;
 	//~ End K2Node Interface.
 
-	/** Get the then output pin */
-	BLUEPRINTGRAPH_API UEdGraphPin* GetThenPin() const;
 	/** Get the return value pin */
 	BLUEPRINTGRAPH_API UEdGraphPin* GetElsePin() const;
 	/** Get the condition pin */
diff --git a/Engine/Source/Editor/BlueprintGraph/Classes/K2Node_SpawnActor.h b/Engine/Source/Editor/BlueprintGraph/Classes/K2Node_SpawnActor.h
--- a/Engine/Source/Editor/BlueprintGraph/Classes/K2Node_SpawnActor.h
+++ b/Engine/Source/Editor/BlueprintGraph/Classes/K2Node_SpawnActor.h
@@ -57,8 +57,6 @@ class UK2Node_SpawnActor : public UK2Node
 	/** See if this is a spawn variable pin, or a 'default' pin */
 	BLUEPRINTGRAPH_API bool IsSpawnVarPin(UEdGraphPin* Pin) const;
 
-	/** Get the then output pin */
-	BLUEPRINTGRAPH_API UEdGraphPin* GetThenPin() const;
 	/** Get the blueprint input pin */	
 	BLUEPRINTGRAPH_API UEdGraphPin* GetBlueprintPin(const TArray<UEdGraphPin*>* InPinsToSearch=NULL) const;
 	/** Get the world context input pin, can return NULL */	
diff --git a/Engine/Source/Editor/BlueprintGraph/Private/K2Node_AssignmentStatement.cpp b/Engine/Source/Editor/BlueprintGraph/Private/K2Node_AssignmentStatement.cpp
--- a/Engine/Source/Editor/BlueprintGraph/Private/K2Node_AssignmentStatement.cpp
+++ b/Engine/Source/Editor/BlueprintGraph/Private/K2Node_AssignmentStatement.cpp
@@ -209,13 +209,6 @@ void UK2Node_AssignmentStatement::PostReconstructNode()
 	Super::PostReconstructNode();
 }
 
-UEdGraphPin* UK2Node_AssignmentStatement::GetThenPin() const
-{
-	UEdGraphPin* Pin = FindPin(UEdGraphSchema_K2::PN_Then);
-	check(Pin != NULL);
-	return Pin;
-}
-
 UEdGraphPin* UK2Node_AssignmentStatement::GetVariablePin() const
 {
 	UEdGraphPin* Pin = FindPin(VariablePinName);
diff --git a/Engine/Source/Editor/BlueprintGraph/Private/K2Node_CallFunction.cpp b/Engine/Source/Editor/BlueprintGraph/Private/K2Node_CallFunction.cpp
--- a/Engine/Source/Editor/BlueprintGraph/Private/K2Node_CallFunction.cpp
+++ b/Engine/Source/Editor/BlueprintGraph/Private/K2Node_CallFunction.cpp
@@ -1403,13 +1403,6 @@ UFunction* UK2Node_CallFunction::GetTargetFunctionFromSkeletonClass() const
 	return TargetFunction;
 }
 
-UEdGraphPin* UK2Node_CallFunction::GetThenPin() const
-{
-	UEdGraphPin* Pin = FindPin(UEdGraphSchema_K2::PN_Then);
-	check(Pin == nullptr || Pin->Direction == EGPD_Output); // If pin exists, it must be output
-	return Pin;
-}
-
 UEdGraphPin* UK2Node_CallFunction::GetReturnValuePin() const
 {
 	UEdGraphPin* Pin = FindPin(UEdGraphSchema_K2::PN_ReturnValue);
diff --git a/Engine/Source/Editor/BlueprintGraph/Private/K2Node_ConstructObjectFromClass.cpp b/Engine/Source/Editor/BlueprintGraph/Private/K2Node_ConstructObjectFromClass.cpp
--- a/Engine/Source/Editor/BlueprintGraph/Private/K2Node_ConstructObjectFromClass.cpp
+++ b/Engine/Source/Editor/BlueprintGraph/Private/K2Node_ConstructObjectFromClass.cpp
@@ -345,13 +345,6 @@ FText UK2Node_ConstructObjectFromClass::GetTooltipText() const
 	return NodeTooltip;
 }
 
-UEdGraphPin* UK2Node_ConstructObjectFromClass::GetThenPin()const
-{
-	UEdGraphPin* Pin = FindPinChecked(UEdGraphSchema_K2::PN_Then);
-	check(Pin->Direction == EGPD_Output);
-	return Pin;
-}
-
 UEdGraphPin* UK2Node_ConstructObjectFromClass::GetClassPin(const TArray<UEdGraphPin*>* InPinsToSearch /*= NULL*/) const
 {
 	const TArray<UEdGraphPin*>* PinsToSearch = InPinsToSearch ? InPinsToSearch : &Pins;
diff --git a/Engine/Source/Editor/BlueprintGraph/Private/K2Node_GetDataTableRow.cpp b/Engine/Source/Editor/BlueprintGraph/Private/K2Node_GetDataTableRow.cpp
--- a/Engine/Source/Editor/BlueprintGraph/Private/K2Node_GetDataTableRow.cpp
+++ b/Engine/Source/Editor/BlueprintGraph/Private/K2Node_GetDataTableRow.cpp
@@ -285,15 +285,6 @@ FText UK2Node_GetDataTableRow::GetTooltipText() const
 	return NodeTooltip;
 }
 
-UEdGraphPin* UK2Node_GetDataTableRow::GetThenPin()const
-{
-	const UEdGraphSchema_K2* K2Schema = GetDefault<UEdGraphSchema_K2>();
-
-	UEdGraphPin* Pin = FindPinChecked(UEdGraphSchema_K2::PN_Then);
-	check(Pin->Direction == EGPD_Output);
-	return Pin;
-}
-
 UEdGraphPin* UK2Node_GetDataTableRow::GetDataTablePin(const TArray<UEdGraphPin*>* InPinsToSearch /*= NULL*/) const
 {
 	const TArray<UEdGraphPin*>* PinsToSearch = InPinsToSearch ? InPinsToSearch : &Pins;
diff --git a/Engine/Source/Editor/BlueprintGraph/Private/K2Node_IfThenElse.cpp b/Engine/Source/Editor/BlueprintGraph/Private/K2Node_IfThenElse.cpp
--- a/Engine/Source/Editor/BlueprintGraph/Private/K2Node_IfThenElse.cpp
+++ b/Engine/Source/Editor/BlueprintGraph/Private/K2Node_IfThenElse.cpp
@@ -141,13 +141,6 @@ FText UK2Node_IfThenElse::GetTooltipText() const
 	return LOCTEXT("BranchStatement_Tooltip", "Branch Statement\nIf Condition is true, execution goes to True, otherwise it goes to False");
 }
 
-UEdGraphPin* UK2Node_IfThenElse::GetThenPin() const
-{
-	UEdGraphPin* Pin = FindPin(UEdGraphSchema_K2::PN_Then);
-	check(Pin);
-	return Pin;
-}
-
 UEdGraphPin* UK2Node_IfThenElse::GetElsePin() const
 {
 	UEdGraphPin* Pin = FindPin(UEdGraphSchema_K2::PN_Else);
diff --git a/Engine/Source/Editor/BlueprintGraph/Private/K2Node_SpawnActor.cpp b/Engine/Source/Editor/BlueprintGraph/Private/K2Node_SpawnActor.cpp
--- a/Engine/Source/Editor/BlueprintGraph/Private/K2Node_SpawnActor.cpp
+++ b/Engine/Source/Editor/BlueprintGraph/Private/K2Node_SpawnActor.cpp
@@ -213,15 +213,6 @@ FText UK2Node_SpawnActor::GetTooltipText() const
 	return NodeTooltip;
 }
 
-UEdGraphPin* UK2Node_SpawnActor::GetThenPin()const
-{
-	const UEdGraphSchema_K2* K2Schema = GetDefault<UEdGraphSchema_K2>();
-
-	UEdGraphPin* Pin = FindPinChecked(UEdGraphSchema_K2::PN_Then);
-	check(Pin->Direction == EGPD_Output);
-	return Pin;
-}
-
 UEdGraphPin* UK2Node_SpawnActor::GetBlueprintPin(const TArray<UEdGraphPin*>* InPinsToSearch /*= NULL*/) const
 {
 	const TArray<UEdGraphPin*>* PinsToSearch = InPinsToSearch ? InPinsToSearch : &Pins;
